#ue4
#ios
* temporarily disable IPP for device detection on PC

TTP332026

diff --git a/Engine/Source/Developer/iOS/IOSTargetPlatform/Private/Windows/IOSDeviceHelperWindows.cpp b/Engine/Source/Developer/iOS/IOSTargetPlatform/Private/Windows/IOSDeviceHelperWindows.cpp
--- a/Engine/Source/Developer/iOS/IOSTargetPlatform/Private/Windows/IOSDeviceHelperWindows.cpp
+++ b/Engine/Source/Developer/iOS/IOSTargetPlatform/Private/Windows/IOSDeviceHelperWindows.cpp
@@ -2,6 +2,8 @@
 
 #include "../IOSTargetPlatformPrivatePCH.h"
 
+#define EABLE_IOS_DEVICE_DETECT 0
+
 struct FDeviceNotificationCallbackInformation
 {
 	FString UDID;
@@ -248,6 +250,7 @@ void FIOSDeviceHelper::Initialize()
 	Event.DeviceName = DummyDeviceName;
 	FIOSDeviceHelper::OnDeviceConnected().Broadcast(Event);
 
+#if ENABLE_IOS_DEVICE_DETECT
 	// add the message pump
 	TickDelegate = FTickerDelegate::CreateStatic(MessageTickDelegate);
 	FTicker::GetCoreTicker().AddTicker(TickDelegate, 5.0f);
@@ -256,6 +259,7 @@ void FIOSDeviceHelper::Initialize()
 	QueryTask = new FDeviceQueryTask();
 	QueryTask->OnDeviceNotification().AddStatic(FIOSDeviceHelper::DeviceCallback);
 	QueryThread = FRunnableThread::Create(QueryTask, TEXT("FIOSDeviceHelper.QueryTask"), false, false, 128 * 1024, TPri_Normal);
+#endif
 }
 
 void FIOSDeviceHelper::DeviceCallback(void* CallbackInfo)
