added a #if #endif guard on a variable that is only created inside of the same #if #endif

diff --git a/Engine/Source/Developer/iOS/IOSTargetPlatform/Private/Windows/IOSDeviceHelperWindows.cpp b/Engine/Source/Developer/iOS/IOSTargetPlatform/Private/Windows/IOSDeviceHelperWindows.cpp
--- a/Engine/Source/Developer/iOS/IOSTargetPlatform/Private/Windows/IOSDeviceHelperWindows.cpp
+++ b/Engine/Source/Developer/iOS/IOSTargetPlatform/Private/Windows/IOSDeviceHelperWindows.cpp
@@ -336,5 +336,7 @@ bool FIOSDeviceHelper::InstallIPAOnDevice(const FTargetDeviceId& DeviceId, const
 
 void FIOSDeviceHelper::EnableDeviceCheck(bool OnOff)
 {
-	QueryTask->Enable(OnOff);
+#if ENABLE_IOS_DEVICE_DETECT
+		QueryTask->Enable(OnOff);
+#endif
 }
\ No newline at end of file
