- Minor fixes to MetalShaderFormat (still not done)

diff --git a/Engine/Source/Developer/iOS/MetalShaderFormat/MetalShaderFormat.Build.cs b/Engine/Source/Developer/iOS/MetalShaderFormat/MetalShaderFormat.Build.cs
--- a/Engine/Source/Developer/iOS/MetalShaderFormat/MetalShaderFormat.Build.cs
+++ b/Engine/Source/Developer/iOS/MetalShaderFormat/MetalShaderFormat.Build.cs
@@ -12,6 +12,7 @@ public class MetalShaderFormat : ModuleRules
 			new string[] {
 				"Core",
 				"ShaderCore",
+				"ShaderCompilerCommon",
 				"ShaderPreprocessor"
 			}
 			);
diff --git a/Engine/Source/Developer/iOS/MetalShaderFormat/Private/MetalShaderCompiler.cpp b/Engine/Source/Developer/iOS/MetalShaderFormat/Private/MetalShaderCompiler.cpp
--- a/Engine/Source/Developer/iOS/MetalShaderFormat/Private/MetalShaderCompiler.cpp
+++ b/Engine/Source/Developer/iOS/MetalShaderFormat/Private/MetalShaderCompiler.cpp
@@ -4,7 +4,7 @@
 #include "MetalShaderFormat.h"
 #include "Core.h"
 #include "../../../EmptyRHI/Public/EmptyShaderResources.h"
-
+#include "ShaderCompilerCommon.h"
 
 #if PLATFORM_WINDOWS
 #include "AllowWindowsPlatformTypes.h"
@@ -82,6 +82,7 @@ static uint32 ParseNumber(const ANSICHAR* &Str)
  */
 static void BuildMetalShaderOutput(
 	FShaderCompilerOutput& ShaderOutput,
+	const FShaderCompilerInput& ShaderInput, 
 	const ANSICHAR* InShaderSource,
 	int32 SourceLen,
 	TArray<FShaderCompilerError>& OutErrors
@@ -92,6 +93,9 @@ static void BuildMetalShaderOutput(
 	FShaderParameterMap& ParameterMap = ShaderOutput.ParameterMap;
 	EShaderFrequency Frequency = (EShaderFrequency)ShaderOutput.Target.Frequency;
 
+	TBitArray<> UsedUniformBufferSlots;
+	UsedUniformBufferSlots.Init(false,32);
+
 	// Write out the magic markers.
 	Header.Frequency = Frequency;
 
@@ -229,6 +233,7 @@ static void BuildMetalShaderOutput(
 			{
 				Header.Bindings.NumUniformBuffers = UBIndex + 1;
 			}
+			UsedUniformBufferSlots[UBIndex] = true;
 			verify(Match(ShaderSource, ')'));
 			ParameterMap.AddParameterAllocation(*BufferName, UBIndex, 0, 0);
 			bHasRegularUniformBuffers = true;
@@ -909,7 +914,7 @@ void CompileShader_Metal(const FShaderCompilerInput& Input,FShaderCompilerOutput
 		{
 			int32 SourceLen = FCStringAnsi::Strlen(MetalShaderSource);
 			Output.Target = Input.Target;
-			BuildMetalShaderOutput(Output, MetalShaderSource, SourceLen, Output.Errors);
+			BuildMetalShaderOutput(Output, Input, MetalShaderSource, SourceLen, Output.Errors);
 		}
 		else
 		{
