- Added a cooker option to generate Metal shaders or not

diff --git a/Engine/Source/Developer/iOS/IOSTargetPlatform/Private/IOSTargetPlatform.cpp b/Engine/Source/Developer/iOS/IOSTargetPlatform/Private/IOSTargetPlatform.cpp
--- a/Engine/Source/Developer/iOS/IOSTargetPlatform/Private/IOSTargetPlatform.cpp
+++ b/Engine/Source/Developer/iOS/IOSTargetPlatform/Private/IOSTargetPlatform.cpp
@@ -207,8 +207,14 @@ void FIOSTargetPlatform::GetAllPossibleShaderFormats( TArray<FName>& OutFormats
 	static FName NAME_OPENGL_ES2_IOS(TEXT("GLSL_ES2_IOS"));
 	static FName NAME_SF_METAL(TEXT("SF_METAL"));
 	OutFormats.AddUnique(NAME_OPENGL_ES2_IOS);
-	// @todo metal enable
-//	OutFormats.AddUnique(NAME_SF_METAL);
+
+	// since Metal is new and fresh, we allow for it to be optional
+	bool bGenerateMetalShaders = false;
+	GConfig->GetBool(TEXT("/Script/UnrealEd.CookerSettings"), TEXT("bGenerateMetalShaders"), bGenerateMetalShaders, GEngineIni);
+	if (bGenerateMetalShaders)
+	{
+		OutFormats.AddUnique(NAME_SF_METAL);
+	}
 }
 
 
diff --git a/Engine/Source/Editor/UnrealEd/Classes/Settings/CookerSettings.h b/Engine/Source/Editor/UnrealEd/Classes/Settings/CookerSettings.h
--- a/Engine/Source/Editor/UnrealEd/Classes/Settings/CookerSettings.h
+++ b/Engine/Source/Editor/UnrealEd/Classes/Settings/CookerSettings.h
@@ -23,4 +23,8 @@ public:
 	/** Quality of 0 means fastest, 4 means best */
 	UPROPERTY(config, EditAnywhere, Category=Textures)
 	int32 DefaultPVRTCQuality;
+
+	/** Whether or not to compile shaders for Metal on IOS (in beta 8). */
+	UPROPERTY(config, EditAnywhere, Category=IOS)
+	bool bGenerateMetalShaders;
 };
