don't allow landscape when portrait is selected and vice-versa
TTP338639
#ue4

diff --git a/Engine/Source/Developer/iOS/IOSPlatformEditor/Private/IOSTargetSettingsCustomization.cpp b/Engine/Source/Developer/iOS/IOSPlatformEditor/Private/IOSTargetSettingsCustomization.cpp
--- a/Engine/Source/Developer/iOS/IOSPlatformEditor/Private/IOSTargetSettingsCustomization.cpp
+++ b/Engine/Source/Developer/iOS/IOSPlatformEditor/Private/IOSTargetSettingsCustomization.cpp
@@ -122,18 +122,19 @@ void FIOSTargetSettingsCustomization::BuildPListSection(IDetailLayoutBuilder& De
 	// Show properties that are gated by the plist being present and writable
 	FSimpleDelegate PlistModifiedDelegate = FSimpleDelegate::CreateRaw(this, &FIOSTargetSettingsCustomization::OnPlistPropertyModified);
 
-#define SETUP_PLIST_PROP(PropName) \
+#define SETUP_PLIST_PROP(PropName, Tip) \
 	{ \
 		TSharedRef<IPropertyHandle> PropertyHandle = DetailLayout.GetProperty(GET_MEMBER_NAME_CHECKED(UIOSRuntimeSettings, PropName)); \
 		PropertyHandle->SetOnPropertyValueChanged(PlistModifiedDelegate); \
 		AppManifestCategory.AddProperty(PropertyHandle) \
-			.EditCondition(SetupForPlatformAttribute, NULL); \
+			.EditCondition(SetupForPlatformAttribute, NULL) \
+			.ToolTip(Tip); \
 	}
 
-	SETUP_PLIST_PROP(bSupportsPortraitOrientation);
-	SETUP_PLIST_PROP(bSupportsUpsideDownOrientation);
-	SETUP_PLIST_PROP(bSupportsLandscapeLeftOrientation);
-	SETUP_PLIST_PROP(bSupportsLandscapeRightOrientation);
+	SETUP_PLIST_PROP(bSupportsPortraitOrientation, TEXT("Supports default portrait orientation. Landscape will not be supported."));
+	SETUP_PLIST_PROP(bSupportsUpsideDownOrientation, TEXT("Supports upside down portrait orientation. Landscape will not be supported."));
+	SETUP_PLIST_PROP(bSupportsLandscapeLeftOrientation, TEXT("Supports left landscape orientation. Protrait will not be supported."));
+	SETUP_PLIST_PROP(bSupportsLandscapeRightOrientation, TEXT("Supports right landscape orientation. Protrait will not be supported."));
 
 #undef SETUP_PLIST_PROP
 }
@@ -228,11 +229,11 @@ void FIOSTargetSettingsCustomization::OnPlistPropertyModified()
 	{
 		OrientationArrayBody += TEXT("\t\t<string>UIInterfaceOrientationPortraitUpsideDown</string>\n");
 	}
-	if (Settings.bSupportsLandscapeLeftOrientation)
+	if (Settings.bSupportsLandscapeLeftOrientation && (!Settings.bSupportsPortraitOrientation && !Settings.bSupportsUpsideDownOrientation))
 	{
 		OrientationArrayBody += TEXT("\t\t<string>UIInterfaceOrientationLandscapeLeft</string>\n");
 	}
-	if (Settings.bSupportsLandscapeRightOrientation)
+	if (Settings.bSupportsLandscapeRightOrientation && (!Settings.bSupportsPortraitOrientation && !Settings.bSupportsUpsideDownOrientation))
 	{
 		OrientationArrayBody += TEXT("\t\t<string>UIInterfaceOrientationLandscapeRight</string>\n");
 	}
