- Fixed up optional Metal support, by having the runtime check the project settings (to see if project wants Metal support or not, before creating a view that supports one or the other)

diff --git a/Engine/Source/Developer/iOS/IOSTargetPlatform/Private/IOSTargetPlatform.cpp b/Engine/Source/Developer/iOS/IOSTargetPlatform/Private/IOSTargetPlatform.cpp
--- a/Engine/Source/Developer/iOS/IOSTargetPlatform/Private/IOSTargetPlatform.cpp
+++ b/Engine/Source/Developer/iOS/IOSTargetPlatform/Private/IOSTargetPlatform.cpp
@@ -206,12 +206,19 @@ void FIOSTargetPlatform::GetAllPossibleShaderFormats( TArray<FName>& OutFormats
 {
 	static FName NAME_OPENGL_ES2_IOS(TEXT("GLSL_ES2_IOS"));
 	static FName NAME_SF_METAL(TEXT("SF_METAL"));
-	OutFormats.AddUnique(NAME_OPENGL_ES2_IOS);
 
-	// since Metal is new and fresh, we allow for it to be optional
-	bool bGenerateMetalShaders = false;
-	GConfig->GetBool(TEXT("/Script/UnrealEd.CookerSettings"), TEXT("bGenerateMetalShaders"), bGenerateMetalShaders, GEngineIni);
-	if (bGenerateMetalShaders)
+	// default to supporting ES2
+	bool bSupportOpenGLES2 = true;
+	GConfig->GetBool(TEXT("/Script/UnrealEd.CookerSettings"), TEXT("bSupportOpenGLES2"), bSupportOpenGLES2, GEngineIni);
+	if (bSupportOpenGLES2)
+	{
+		OutFormats.AddUnique(NAME_OPENGL_ES2_IOS);
+	}
+
+	// default to NOT supporting metal
+	bool bSupportMetal = false;
+	GConfig->GetBool(TEXT("/Script/UnrealEd.CookerSettings"), TEXT("bSupportMetal"), bSupportMetal, GEngineIni);
+	if (bSupportMetal)
 	{
 		OutFormats.AddUnique(NAME_SF_METAL);
 	}
diff --git a/Engine/Source/Editor/UnrealEd/Classes/Settings/CookerSettings.h b/Engine/Source/Editor/UnrealEd/Classes/Settings/CookerSettings.h
--- a/Engine/Source/Editor/UnrealEd/Classes/Settings/CookerSettings.h
+++ b/Engine/Source/Editor/UnrealEd/Classes/Settings/CookerSettings.h
@@ -24,7 +24,11 @@ public:
 	UPROPERTY(config, EditAnywhere, Category=Textures)
 	int32 DefaultPVRTCQuality;
 
-	/** Whether or not to compile shaders for Metal on IOS (in beta 8). */
+	/** Whether or not to add support for Metal API (requires IOS8 and A7 processors) */
 	UPROPERTY(config, EditAnywhere, Category=IOS)
-	bool bGenerateMetalShaders;
+	bool bSupportMetal;
+
+	/** Whether or not to add support for OpenGL ES2 (if this is false, then your game should specify minimum IOS8 version and A7, in the plist) */
+	UPROPERTY(config, EditAnywhere, Category = IOS)
+	bool bSupportOpenGLES2;
 };
diff --git a/Engine/Source/Editor/UnrealEd/Private/Commandlets/CookCommandlet.cpp b/Engine/Source/Editor/UnrealEd/Private/Commandlets/CookCommandlet.cpp
--- a/Engine/Source/Editor/UnrealEd/Private/Commandlets/CookCommandlet.cpp
+++ b/Engine/Source/Editor/UnrealEd/Private/Commandlets/CookCommandlet.cpp
@@ -29,6 +29,8 @@ UCookerSettings::UCookerSettings(const FPostConstructInitializeProperties& PCIP)
 	: Super(PCIP)
 {
 	DefaultPVRTCQuality = 0;
+	bSupportMetal = false;
+	bSupportOpenGLES2 = true;
 }
 
 
diff --git a/Engine/Source/Programs/UnrealBuildTool/IOS/UEBuildIOS.cs b/Engine/Source/Programs/UnrealBuildTool/IOS/UEBuildIOS.cs
--- a/Engine/Source/Programs/UnrealBuildTool/IOS/UEBuildIOS.cs
+++ b/Engine/Source/Programs/UnrealBuildTool/IOS/UEBuildIOS.cs
@@ -122,13 +122,12 @@ namespace UnrealBuildTool
 			}
 
 			// needs IOS8 for Metal
-			// @todo metal enable: Uncomment the following lines if you want to give Metal a try!
-// 			if (IOSToolChain.IOSSDKVersionFloat >= 8.0 && UEBuildConfiguration.bCompileAgainstEngine)
-// 			{
-// 				InBuildTarget.GlobalCompileEnvironment.Config.Definitions.Add("HAS_METAL=1");
-// 				InBuildTarget.ExtraModuleNames.Add("MetalRHI");
-// 			}
-// 			else
+			if (IOSToolChain.IOSSDKVersionFloat >= 8.0 && UEBuildConfiguration.bCompileAgainstEngine)
+			{
+				InBuildTarget.GlobalCompileEnvironment.Config.Definitions.Add("HAS_METAL=1");
+				InBuildTarget.ExtraModuleNames.Add("MetalRHI");
+			}
+			else
 			{
 				InBuildTarget.GlobalCompileEnvironment.Config.Definitions.Add("HAS_METAL=0");
 			}
diff --git a/Engine/Source/Runtime/RHI/Private/Apple/AppleDynamicRHI.cpp b/Engine/Source/Runtime/RHI/Private/Apple/AppleDynamicRHI.cpp
--- a/Engine/Source/Runtime/RHI/Private/Apple/AppleDynamicRHI.cpp
+++ b/Engine/Source/Runtime/RHI/Private/Apple/AppleDynamicRHI.cpp
@@ -11,8 +11,8 @@ FDynamicRHI* PlatformCreateDynamicRHI()
 	IDynamicRHIModule* DynamicRHIModule = NULL;
 
 	// Load the dynamic RHI module.
-#if PLATFORM_IOS && HAS_METAL
-	if (FPlatformMisc::HasPlatformFeature(TEXT("Metal")) && !FParse::Param(FCommandLine::Get(), TEXT("ES2")))
+#if PLATFORM_IOS
+	if (FPlatformMisc::HasPlatformFeature(TEXT("Metal")))
 	{
 		DynamicRHIModule = &FModuleManager::LoadModuleChecked<IDynamicRHIModule>(TEXT("MetalRHI"));
 	}
