change the plist editor to looking for and saving to Info.plist
#ue4
#ios

diff --git a/Engine/Source/Developer/iOS/IOSPlatformEditor/Private/IOSTargetSettingsCustomization.cpp b/Engine/Source/Developer/iOS/IOSPlatformEditor/Private/IOSTargetSettingsCustomization.cpp
--- a/Engine/Source/Developer/iOS/IOSPlatformEditor/Private/IOSTargetSettingsCustomization.cpp
+++ b/Engine/Source/Developer/iOS/IOSPlatformEditor/Private/IOSTargetSettingsCustomization.cpp
@@ -26,7 +26,7 @@ TSharedRef<IDetailCustomization> FIOSTargetSettingsCustomization::MakeInstance()
 
 FIOSTargetSettingsCustomization::FIOSTargetSettingsCustomization()
 	: EngineInfoPath(FString::Printf(TEXT("%sBuild/IOS/UE4Game-Info.plist"), *FPaths::EngineDir()))
-	, GameInfoPath(FString::Printf(TEXT("%sBuild/IOS/%s-Info.plist"), *FPaths::GameDir(), GGameName))
+	, GameInfoPath(FString::Printf(TEXT("%sBuild/IOS/Info.plist"), *FPaths::GameDir()))
 	, EngineGraphicsPath(FString::Printf(TEXT("%sBuild/IOS/Resources/Graphics"), *FPaths::EngineDir()))
 	, GameGraphicsPath(FString::Printf(TEXT("%sBuild/IOS/Resources/Graphics"), *FPaths::GameDir()))
 {
diff --git a/Engine/Source/Editor/PListEditor/Private/SPlistEditor.cpp b/Engine/Source/Editor/PListEditor/Private/SPlistEditor.cpp
--- a/Engine/Source/Editor/PListEditor/Private/SPlistEditor.cpp
+++ b/Engine/Source/Editor/PListEditor/Private/SPlistEditor.cpp
@@ -822,6 +822,13 @@ bool SPListEditorPanel::OpenFile(FString FilePath)
 	bool bLoadResult = Doc.LoadFile(FilePath);
 
 	FString OutError;
+	if( !bLoadResult )
+	{
+		// try Info.plist
+		FilePath = FilePath.Replace(*(FString(GGameName) + TEXT("-")), TEXT(""));
+		bLoadResult = Doc.LoadFile(FilePath);
+	}
+
 	if( !bLoadResult || !ParseXmlTree(Doc, OutError) )
 	{
 		DisplayNotification( LOCTEXT("PListLoadFail", "Load Failed"), NTF_Fail );
