addition of Bundle Name, Bundle Display Name, Bundle Identifier, and Version to the iOS project settings.
TTP341741
#ue4
#ios

diff --git a/Engine/Source/Developer/iOS/IOSPlatformEditor/Private/IOSTargetSettingsCustomization.cpp b/Engine/Source/Developer/iOS/IOSPlatformEditor/Private/IOSTargetSettingsCustomization.cpp
--- a/Engine/Source/Developer/iOS/IOSPlatformEditor/Private/IOSTargetSettingsCustomization.cpp
+++ b/Engine/Source/Developer/iOS/IOSPlatformEditor/Private/IOSTargetSettingsCustomization.cpp
@@ -131,6 +131,10 @@ void FIOSTargetSettingsCustomization::BuildPListSection(IDetailLayoutBuilder& De
 			.ToolTip(Tip); \
 	}
 
+	SETUP_PLIST_PROP(BundleDisplayName, TEXT("Specifies the the display name for the application. This will be displayed under the icon on the device."));
+	SETUP_PLIST_PROP(BundleName, TEXT("Specifies the the name of the application bundle. This is the short name for the application bundle."));
+	SETUP_PLIST_PROP(BundleIdentifier, TEXT("Specifies the bundle identifier for the application."));
+	SETUP_PLIST_PROP(VersionInfo, TEXT("Specifies the version for the application."));
 	SETUP_PLIST_PROP(bSupportsPortraitOrientation, TEXT("Supports default portrait orientation. Landscape will not be supported."));
 	SETUP_PLIST_PROP(bSupportsUpsideDownOrientation, TEXT("Supports upside down portrait orientation. Landscape will not be supported."));
 	SETUP_PLIST_PROP(bSupportsLandscapeLeftOrientation, TEXT("Supports left landscape orientation. Protrait will not be supported."));
@@ -241,6 +245,27 @@ void FIOSTargetSettingsCustomization::OnPlistPropertyModified()
 
 	Updater.ReplaceKey(InterfaceOrientations, ClosingArray, OrientationArrayBody);
 
+	// build the replacement bundle display name
+	const FString BundleDisplayNameKey(TEXT("<key>CFBundleDisplayName</key>"));
+	const FString ClosingString(TEXT("</string>"));
+	FString BundleDisplayNameBody = TEXT("\n\t<string>") + Settings.BundleDisplayName;
+	Updater.ReplaceKey(BundleDisplayNameKey, ClosingString, BundleDisplayNameBody);
+
+	// build the replacement bundle display name
+	const FString BundleNameKey(TEXT("<key>CFBundleName</key>"));
+	FString BundleNameBody = TEXT("\n\t<string>") + Settings.BundleName;
+	Updater.ReplaceKey(BundleNameKey, ClosingString, BundleNameBody);
+
+	// build the replacement bundle identifier
+	const FString BundleIdentifierKey(TEXT("<key>CFBundleIdentifier</key>"));
+	FString BundleIdentifierBody = TEXT("\n\t<string>") + Settings.BundleIdentifier;
+	Updater.ReplaceKey(BundleIdentifierKey, ClosingString, BundleIdentifierBody);
+
+	// build the replacement version info
+	const FString BuBundleShortVersionKey(TEXT("<key>CFBundleShortVersionString</key>"));
+	FString VersionInfoBody = TEXT("\n\t<string>") + Settings.VersionInfo;
+	Updater.ReplaceKey(BuBundleShortVersionKey, ClosingString, VersionInfoBody);
+
 	// Write out the updated .plist
 	Updater.Finalize(GameInfoPath, true, FFileHelper::EEncodingOptions::ForceUTF8);
 }
diff --git a/Engine/Source/Runtime/IOS/IOSRuntimeSettings/Classes/IOSRuntimeSettings.h b/Engine/Source/Runtime/IOS/IOSRuntimeSettings/Classes/IOSRuntimeSettings.h
--- a/Engine/Source/Runtime/IOS/IOSRuntimeSettings/Classes/IOSRuntimeSettings.h
+++ b/Engine/Source/Runtime/IOS/IOSRuntimeSettings/Classes/IOSRuntimeSettings.h
@@ -34,6 +34,22 @@ public:
 	UPROPERTY(GlobalConfig, EditAnywhere, Category = DeviceOrientations)
 	uint32 bSupportsLandscapeRightOrientation : 1;
 
+	// Bundle Display Name
+	UPROPERTY(GlobalConfig, EditAnywhere, Category = BundleInfo)
+	FString BundleDisplayName;
+
+	// Bundle Name
+	UPROPERTY(GlobalConfig, EditAnywhere, Category = BundleInfo)
+	FString BundleName;
+
+	// Bundle identifier
+	UPROPERTY(GlobalConfig, EditAnywhere, Category = BundleInfo)
+	FString BundleIdentifier;
+
+	// version info
+	UPROPERTY(GlobalConfig, EditAnywhere, Category = BundleInfo)
+	FString VersionInfo;
+
 #if WITH_EDITOR
 	// UObject interface
 	virtual void PostEditChangeProperty(struct FPropertyChangedEvent& PropertyChangedEvent) override;
diff --git a/Engine/Source/Runtime/IOS/IOSRuntimeSettings/Private/IOSRuntimeSettings.cpp b/Engine/Source/Runtime/IOS/IOSRuntimeSettings/Private/IOSRuntimeSettings.cpp
--- a/Engine/Source/Runtime/IOS/IOSRuntimeSettings/Private/IOSRuntimeSettings.cpp
+++ b/Engine/Source/Runtime/IOS/IOSRuntimeSettings/Private/IOSRuntimeSettings.cpp
@@ -9,6 +9,10 @@ UIOSRuntimeSettings::UIOSRuntimeSettings(const class FPostConstructInitializePro
 {
 	bEnableGameCenterSupport = true;
 	bSupportsPortraitOrientation = true;
+	BundleDisplayName = TEXT("UE4 Game");
+	BundleName = TEXT("MyUE4Game");
+	BundleIdentifier = TEXT("com.YourCompany.GameNameNoSpaces");
+	VersionInfo = TEXT("1.0.0");
 }
 
 #if WITH_EDITOR
