UE4 - Fix text render actors showing blocky on mobile (TTP 342605)

diff --git a/Engine/Source/Developer/Android/AndroidTargetPlatform/Private/AndroidTargetPlatform.inl b/Engine/Source/Developer/Android/AndroidTargetPlatform/Private/AndroidTargetPlatform.inl
--- a/Engine/Source/Developer/Android/AndroidTargetPlatform/Private/AndroidTargetPlatform.inl
+++ b/Engine/Source/Developer/Android/AndroidTargetPlatform/Private/AndroidTargetPlatform.inl
@@ -226,6 +226,10 @@ inline void FAndroidTargetPlatform<TPlatformProperties>::GetTextureFormats( cons
 	{
 		OutFormats.Add(AndroidTexFormat::NameG8);
 	}
+	else if (InTexture->CompressionSettings == TC_DistanceFieldFont)
+	{
+		OutFormats.Add(AndroidTexFormat::NameG8);
+	}
 	else if (InTexture->bForcePVRTC4)
 	{
 		AddTextureFormatIfSupports(AndroidTexFormat::NamePVRTC4, OutFormats);
diff --git a/Engine/Source/Developer/HTML5/HTML5TargetPlatform/Private/HTML5TargetPlatform.cpp b/Engine/Source/Developer/HTML5/HTML5TargetPlatform/Private/HTML5TargetPlatform.cpp
--- a/Engine/Source/Developer/HTML5/HTML5TargetPlatform/Private/HTML5TargetPlatform.cpp
+++ b/Engine/Source/Developer/HTML5/HTML5TargetPlatform/Private/HTML5TargetPlatform.cpp
@@ -172,6 +172,10 @@ void FHTML5TargetPlatform::GetTextureFormats( const UTexture* Texture, TArray<FN
 	{
 		TextureFormatName = NameDXT5;
 	}
+	else if (Texture->CompressionSettings == TC_DistanceFieldFont)
+	{
+		TextureFormatName = NameG8;
+	}
 	else if (Texture->CompressionNoAlpha)
 	{
 		TextureFormatName = NameDXT1;
diff --git a/Engine/Source/Developer/iOS/IOSTargetPlatform/Private/IOSTargetPlatform.cpp b/Engine/Source/Developer/iOS/IOSTargetPlatform/Private/IOSTargetPlatform.cpp
--- a/Engine/Source/Developer/iOS/IOSTargetPlatform/Private/IOSTargetPlatform.cpp
+++ b/Engine/Source/Developer/iOS/IOSTargetPlatform/Private/IOSTargetPlatform.cpp
@@ -21,11 +21,11 @@ FIOSTargetPlatform::FIOSTargetPlatform()
 	// Initialize Ticker for device discovery
 	TickDelegate = FTickerDelegate::CreateRaw(this, &FIOSTargetPlatform::HandleTicker);
 	FTicker::GetCoreTicker().AddTicker(TickDelegate, 10.0f);
-    
-    // initialize the connected device detector
-    DeviceHelper.OnDeviceConnected().AddRaw(this, &FIOSTargetPlatform::HandleDeviceConnected);
-    DeviceHelper.OnDeviceDisconnected().AddRaw(this, &FIOSTargetPlatform::HandleDeviceDisconnected);
-    DeviceHelper.Initialize();
+	
+	// initialize the connected device detector
+	DeviceHelper.OnDeviceConnected().AddRaw(this, &FIOSTargetPlatform::HandleDeviceConnected);
+	DeviceHelper.OnDeviceDisconnected().AddRaw(this, &FIOSTargetPlatform::HandleDeviceDisconnected);
+	DeviceHelper.Initialize();
 }
 
 
@@ -221,13 +221,13 @@ void FIOSTargetPlatform::HandleDeviceConnected(const FIOSLaunchDaemonPong& Messa
 {
 	FTargetDeviceId DeviceId;
 	FTargetDeviceId::Parse(Message.DeviceID, DeviceId);
-    
+	
 	FIOSTargetDevicePtr& Device = Devices.FindOrAdd(DeviceId);
-    
+	
 	if (!Device.IsValid())
 	{
 		Device = MakeShareable(new FIOSTargetDevice(*this));
-        
+		
 		Device->SetFeature(ETargetDeviceFeatures::Reboot, Message.bCanReboot);
 		Device->SetFeature(ETargetDeviceFeatures::PowerOn, Message.bCanPowerOn);
 		Device->SetFeature(ETargetDeviceFeatures::PowerOff, Message.bCanPowerOff);
@@ -235,10 +235,10 @@ void FIOSTargetPlatform::HandleDeviceConnected(const FIOSLaunchDaemonPong& Messa
 		Device->SetDeviceName(Message.DeviceName);
 		Device->SetDeviceType(Message.DeviceType);
 		Device->SetIsSimulated(Message.DeviceID.Contains(TEXT("Simulator")));
-        
+		
 		DeviceDiscoveredEvent.Broadcast(Device.ToSharedRef());
 	}
-    
+	
 	// Add a very long time period to prevent the devices from getting disconnected due to a lack of pong messages
 	Device->LastPinged = FDateTime::UtcNow() + FTimespan(100, 0, 0, 0, 0);
 }
@@ -248,13 +248,13 @@ void FIOSTargetPlatform::HandleDeviceDisconnected(const FIOSLaunchDaemonPong& Me
 {
 	FTargetDeviceId DeviceId;
 	FTargetDeviceId::Parse(Message.DeviceID, DeviceId);
-    
+	
 	FIOSTargetDevicePtr& Device = Devices.FindOrAdd(DeviceId);
-    
+	
 	if (Device.IsValid())
 	{
-        DeviceLostEvent.Broadcast(Device.ToSharedRef());
-        Devices.Remove(DeviceId);
+		DeviceLostEvent.Broadcast(Device.ToSharedRef());
+		Devices.Remove(DeviceId);
 	}
 }
 
@@ -378,6 +378,10 @@ void FIOSTargetPlatform::GetTextureFormats( const UTexture* Texture, TArray<FNam
 	{
 		TextureFormatName = NameG8;
 	}
+	else if (Texture->CompressionSettings == TC_DistanceFieldFont)
+	{
+		TextureFormatName = NameG8;
+	}
 	else if (Texture->bForcePVRTC4)
 	{
 		TextureFormatName = NamePVRTC4;
