- Fixed the path to the metal offline compiler for beta4+

diff --git a/Engine/Source/Developer/iOS/MetalShaderFormat/Private/MetalShaderCompiler.cpp b/Engine/Source/Developer/iOS/MetalShaderFormat/Private/MetalShaderCompiler.cpp
--- a/Engine/Source/Developer/iOS/MetalShaderFormat/Private/MetalShaderCompiler.cpp
+++ b/Engine/Source/Developer/iOS/MetalShaderFormat/Private/MetalShaderCompiler.cpp
@@ -649,7 +649,7 @@ static void BuildMetalShaderOutput(
 
 		// metal commandlines
 		FString Params = FString::Printf(TEXT("-std=ios-metal1.0 %s -o %s"), *InputFilename, *ObjFilename);
-		FPlatformProcess::ExecProcess( TEXT("/Applications/Xcode.app/Contents/Developer/Platforms/iPhoneOS.platform/usr/local/bin/metal"), *Params, &ReturnCode, &Results, &Errors );
+		FPlatformProcess::ExecProcess( TEXT("/Applications/Xcode.app/Contents/Developer/Platforms/iPhoneOS.platform/usr/bin/metal"), *Params, &ReturnCode, &Results, &Errors );
 
 		// handle compile error
 		if (ReturnCode != 0 || IFileManager::Get().FileSize(*ObjFilename) <= 0)
@@ -662,7 +662,7 @@ static void BuildMetalShaderOutput(
 		else
 		{
 			Params = FString::Printf(TEXT("r %s %s"), *ArFilename, *ObjFilename);
-			FPlatformProcess::ExecProcess( TEXT("/Applications/Xcode.app/Contents/Developer/Platforms/iPhoneOS.platform/usr/local/bin/metal-ar"), *Params, &ReturnCode, &Results, &Errors );
+			FPlatformProcess::ExecProcess( TEXT("/Applications/Xcode.app/Contents/Developer/Platforms/iPhoneOS.platform/usr/bin/metal-ar"), *Params, &ReturnCode, &Results, &Errors );
 
 			// handle compile error
 			if (ReturnCode != 0 || IFileManager::Get().FileSize(*ArFilename) <= 0)
@@ -675,7 +675,7 @@ static void BuildMetalShaderOutput(
 			else
 			{
 				Params = FString::Printf(TEXT("-o %s %s"), *OutputFilename, *ArFilename);
-				FPlatformProcess::ExecProcess( TEXT("/Applications/Xcode.app/Contents/Developer/Platforms/iPhoneOS.platform/usr/local/bin/metallib"), *Params, &ReturnCode, &Results, &Errors );
+				FPlatformProcess::ExecProcess( TEXT("/Applications/Xcode.app/Contents/Developer/Platforms/iPhoneOS.platform/usr/bin/metallib"), *Params, &ReturnCode, &Results, &Errors );
 		
 				// handle compile error
 				if (ReturnCode != 0 || IFileManager::Get().FileSize(*OutputFilename) <= 0)
diff --git a/Engine/Source/Developer/iOS/MetalShaderFormat/Private/MetalShaderFormat.cpp b/Engine/Source/Developer/iOS/MetalShaderFormat/Private/MetalShaderFormat.cpp
--- a/Engine/Source/Developer/iOS/MetalShaderFormat/Private/MetalShaderFormat.cpp
+++ b/Engine/Source/Developer/iOS/MetalShaderFormat/Private/MetalShaderFormat.cpp
@@ -14,7 +14,7 @@ class FMetalShaderFormat : public IShaderFormat
 {
 	enum
 	{
-		HEADER_VERSION = 4,
+		HEADER_VERSION = 5,
 	};
 	struct FVersion
 	{
