fix for logging out the IPP validation
fix for no quotes around the path when running the validation
#ue4
#ios

diff --git a/Engine/Source/Developer/iOS/IOSTargetPlatform/Private/IOSTargetPlatform.cpp b/Engine/Source/Developer/iOS/IOSTargetPlatform/Private/IOSTargetPlatform.cpp
--- a/Engine/Source/Developer/iOS/IOSTargetPlatform/Private/IOSTargetPlatform.cpp
+++ b/Engine/Source/Developer/iOS/IOSTargetPlatform/Private/IOSTargetPlatform.cpp
@@ -99,6 +99,13 @@ bool FIOSTargetPlatform::IsSdkInstalled(bool bProjectHasCode, FString& OutDocume
 	return biOSSDKInstalled;
 }
 
+static FString OutputMessage;
+static void OnOutput(FString Message)
+{
+	OutputMessage += Message;
+	UE_LOG(LogTemp, Display, TEXT("%s"), *Message);
+}
+
 int FIOSTargetPlatform::DoesntHaveRequirements(const FString& ProjectPath, bool bProjectHasCode, FString& OutDocumentationPath) const
 {
 	int bReadyToBuild = ETargetPlatformReadyStatus::Ready; // @todo How do we check that the iOS SDK is installed when building from Windows? Is that even possible?
@@ -122,8 +129,10 @@ int FIOSTargetPlatform::DoesntHaveRequirements(const FString& ProjectPath, bool
 
 	// shell to IPP and get the status of the provision and cert
 	FString CmdExe = FPaths::ConvertRelativePathToFull(FPaths::EngineDir() / TEXT("Binaries/DotNet/IOS/IPhonePackager.exe"));
-	FString CommandLine = FString::Printf(TEXT("Validate Engine -project %s"), *ProjectPath);
+	FString CommandLine = FString::Printf(TEXT("Validate Engine -project \"%s\""), *ProjectPath);
 	TSharedPtr<FMonitoredProcess> IPPProcess = MakeShareable(new FMonitoredProcess(CmdExe, CommandLine, true));
+	OutputMessage = TEXT("");
+	IPPProcess->OnOutput().BindStatic(&OnOutput);
 	IPPProcess->Launch();
 	while(IPPProcess->IsRunning())
 	{
