updated metal setting to update the Info.plist
adding Info.plist for Zen
#zen

diff --git a/Engine/Source/Developer/iOS/IOSPlatformEditor/Private/IOSTargetSettingsCustomization.cpp b/Engine/Source/Developer/iOS/IOSPlatformEditor/Private/IOSTargetSettingsCustomization.cpp
--- a/Engine/Source/Developer/iOS/IOSPlatformEditor/Private/IOSTargetSettingsCustomization.cpp
+++ b/Engine/Source/Developer/iOS/IOSPlatformEditor/Private/IOSTargetSettingsCustomization.cpp
@@ -66,6 +66,9 @@ void FIOSTargetSettingsCustomization::BuildPListSection(IDetailLayoutBuilder& De
 {
 	// Info.plist category
 	IDetailCategoryBuilder& AppManifestCategory = DetailLayout.EditCategory(TEXT("Info.plist"));
+	IDetailCategoryBuilder& BundleCategory = DetailLayout.EditCategory(TEXT("Bundle Information"));
+	IDetailCategoryBuilder& OrientationCategory = DetailLayout.EditCategory(TEXT("Orientation"));
+	IDetailCategoryBuilder& RenderCategory = DetailLayout.EditCategory(TEXT("Rendering"));
 
 	TSharedRef<SPlatformSetupMessage> PlatformSetupMessage = SNew(SPlatformSetupMessage, GameInfoPath)
 		.PlatformName(LOCTEXT("iOSPlatformName", "iOS"))
@@ -122,23 +125,26 @@ void FIOSTargetSettingsCustomization::BuildPListSection(IDetailLayoutBuilder& De
 	// Show properties that are gated by the plist being present and writable
 	FSimpleDelegate PlistModifiedDelegate = FSimpleDelegate::CreateRaw(this, &FIOSTargetSettingsCustomization::OnPlistPropertyModified);
 
-#define SETUP_PLIST_PROP(PropName, Tip) \
+#define SETUP_PLIST_PROP(PropName, Category, Tip) \
 	{ \
 		TSharedRef<IPropertyHandle> PropertyHandle = DetailLayout.GetProperty(GET_MEMBER_NAME_CHECKED(UIOSRuntimeSettings, PropName)); \
 		PropertyHandle->SetOnPropertyValueChanged(PlistModifiedDelegate); \
-		AppManifestCategory.AddProperty(PropertyHandle) \
+		Category.AddProperty(PropertyHandle) \
 			.EditCondition(SetupForPlatformAttribute, NULL) \
 			.ToolTip(Tip); \
 	}
 
-	SETUP_PLIST_PROP(BundleDisplayName, TEXT("Specifies the the display name for the application. This will be displayed under the icon on the device."));
-	SETUP_PLIST_PROP(BundleName, TEXT("Specifies the the name of the application bundle. This is the short name for the application bundle."));
-	SETUP_PLIST_PROP(BundleIdentifier, TEXT("Specifies the bundle identifier for the application."));
-	SETUP_PLIST_PROP(VersionInfo, TEXT("Specifies the version for the application."));
-	SETUP_PLIST_PROP(bSupportsPortraitOrientation, TEXT("Supports default portrait orientation. Landscape will not be supported."));
-	SETUP_PLIST_PROP(bSupportsUpsideDownOrientation, TEXT("Supports upside down portrait orientation. Landscape will not be supported."));
-	SETUP_PLIST_PROP(bSupportsLandscapeLeftOrientation, TEXT("Supports left landscape orientation. Protrait will not be supported."));
-	SETUP_PLIST_PROP(bSupportsLandscapeRightOrientation, TEXT("Supports right landscape orientation. Protrait will not be supported."));
+	SETUP_PLIST_PROP(BundleDisplayName, BundleCategory, TEXT("Specifies the the display name for the application. This will be displayed under the icon on the device."));
+	SETUP_PLIST_PROP(BundleName, BundleCategory, TEXT("Specifies the the name of the application bundle. This is the short name for the application bundle."));
+	SETUP_PLIST_PROP(BundleIdentifier, BundleCategory, TEXT("Specifies the bundle identifier for the application."));
+	SETUP_PLIST_PROP(VersionInfo, BundleCategory, TEXT("Specifies the version for the application."));
+	SETUP_PLIST_PROP(bSupportsPortraitOrientation, OrientationCategory, TEXT("Supports default portrait orientation. Landscape will not be supported."));
+	SETUP_PLIST_PROP(bSupportsUpsideDownOrientation, OrientationCategory, TEXT("Supports upside down portrait orientation. Landscape will not be supported."));
+	SETUP_PLIST_PROP(bSupportsLandscapeLeftOrientation, OrientationCategory, TEXT("Supports left landscape orientation. Protrait will not be supported."));
+	SETUP_PLIST_PROP(bSupportsLandscapeRightOrientation, OrientationCategory, TEXT("Supports right landscape orientation. Protrait will not be supported."));
+	
+	SETUP_PLIST_PROP(bSupportsMetal, RenderCategory, TEXT("Whether or not to add support for Metal API (requires IOS8 and A7 processors)."));
+	SETUP_PLIST_PROP(bSupportsOpenGLES2, RenderCategory, TEXT("Whether or not to add support for OpenGL ES2 (if this is false, then your game should specify minimum IOS8 version and use \"metal\" instead of \"opengles-2\" in UIRequiredDeviceCapabilities)"));
 
 #undef SETUP_PLIST_PROP
 }
@@ -242,7 +248,6 @@ void FIOSTargetSettingsCustomization::OnPlistPropertyModified()
 		OrientationArrayBody += TEXT("\t\t<string>UIInterfaceOrientationLandscapeRight</string>\n");
 	}
 	OrientationArrayBody += TEXT("\t");
-
 	Updater.ReplaceKey(InterfaceOrientations, ClosingArray, OrientationArrayBody);
 
 	// build the replacement bundle display name
@@ -262,9 +267,25 @@ void FIOSTargetSettingsCustomization::OnPlistPropertyModified()
 	Updater.ReplaceKey(BundleIdentifierKey, ClosingString, BundleIdentifierBody);
 
 	// build the replacement version info
-	const FString BuBundleShortVersionKey(TEXT("<key>CFBundleShortVersionString</key>"));
+	const FString BundleShortVersionKey(TEXT("<key>CFBundleShortVersionString</key>"));
 	FString VersionInfoBody = TEXT("\n\t<string>") + Settings.VersionInfo;
-	Updater.ReplaceKey(BuBundleShortVersionKey, ClosingString, VersionInfoBody);
+	Updater.ReplaceKey(BundleShortVersionKey, ClosingString, VersionInfoBody);
+
+	// build the replacement required device caps
+	const FString RequiredDeviceCaps(TEXT("<key>UIRequiredDeviceCapabilities</key>"));
+	FString DeviceCapsArrayBody = TEXT("\n\t<array>\n");
+	// automatically add armv7 for now
+	DeviceCapsArrayBody += TEXT("\t\t<string>armv7</string>\n");
+	if (Settings.bSupportsOpenGLES2)
+	{
+		DeviceCapsArrayBody += TEXT("\t\t<string>opengles-2</string>");
+	}
+	else if (Settings.bSupportsMetal)
+	{
+		DeviceCapsArrayBody += TEXT("\t\t<string>metal</string>");
+	}
+	DeviceCapsArrayBody += TEXT("\t");
+	Updater.ReplaceKey(RequiredDeviceCaps, ClosingArray, DeviceCapsArrayBody);
 
 	// Write out the updated .plist
 	Updater.Finalize(GameInfoPath, true, FFileHelper::EEncodingOptions::ForceUTF8);
