Fixed Mac CIS build

diff --git a/Engine/Source/Developer/iOS/IOSTargetPlatform/Private/IOSTargetPlatform.cpp b/Engine/Source/Developer/iOS/IOSTargetPlatform/Private/IOSTargetPlatform.cpp
--- a/Engine/Source/Developer/iOS/IOSTargetPlatform/Private/IOSTargetPlatform.cpp
+++ b/Engine/Source/Developer/iOS/IOSTargetPlatform/Private/IOSTargetPlatform.cpp
@@ -80,7 +80,7 @@ bool FIOSTargetPlatform::IsSdkInstalled(bool bProjectHasCode, FString& OutTutori
 {
 	bool biOSSDKInstalled = true; // @todo How do we check that the iOS SDK is installed when building from Windows? Is that even possible?
 #if PLATFORM_MAC
-	OutDocumentationPath = FString("Shared/Tutorials/InstallingXCodeTutorial");
+	OutTutorialPath = FString("Shared/Tutorials/InstallingXCodeTutorial");
 	biOSSDKInstalled = IFileManager::Get().DirectoryExists(TEXT("/Applications/Xcode.app/Contents/Developer/Platforms/iPhoneOS.platform"));
 #else
 	{
