UE4 - hlslcc - Metal backend from ThirdParty to Engine module

diff --git a/Engine/Source/Developer/iOS/MetalShaderFormat/Private/MetalBackend.cpp b/Engine/Source/Developer/iOS/MetalShaderFormat/Private/MetalBackend.cpp
new file mode 100644
--- /dev/null
+++ b/Engine/Source/Developer/iOS/MetalShaderFormat/Private/MetalBackend.cpp
@@ -0,0 +1,2694 @@
+// Copyright 1998-2014 Epic Games, Inc. All Rights Reserved.
+
+#include "MetalShaderFormat.h"
+#include "hlslcc.h"
+#include "hlslcc_private.h"
+#include "MetalBackend.h"
+#include "glsl/ir_gen_glsl.h"
+#include "mesa/compiler.h"
+#include "mesa/glsl_parser_extras.h"
+#include "mesa/hash_table.h"
+#include "mesa/ir_rvalue_visitor.h"
+#include "PackUniformBuffers.h"
+#include "IRDump.h"
+#include "OptValueNumbering.h"
+#include "mesa/ir_optimization.h"
+#include "MetalUtils.h"
+
+
+/**
+ * This table must match the ir_expression_operation enum.
+ */
+static const char * const GLSLExpressionTable[ir_opcode_count][4] =
+	{
+	{ "(~", ")", "", "" }, // ir_unop_bit_not,
+	{ "not(", ")", "", "!" }, // ir_unop_logic_not,
+	{ "(-", ")", "", "" }, // ir_unop_neg,
+	{ "fabs(", ")", "", "" }, // ir_unop_abs,
+	{ "sign(", ")", "", "" }, // ir_unop_sign,
+	{ "(1.0/(", "))", "", "" }, // ir_unop_rcp,
+	{ "rsqrt(", ")", "", "" }, // ir_unop_rsq,
+	{ "sqrt(", ")", "", "" }, // ir_unop_sqrt,
+	{ "exp(", ")", "", "" }, // ir_unop_exp,      /**< Log base e on gentype */
+	{ "log(", ")", "", "" }, // ir_unop_log,	     /**< Natural log on gentype */
+	{ "exp2(", ")", "", "" }, // ir_unop_exp2,
+	{ "log2(", ")", "", "" }, // ir_unop_log2,
+	{ "int(", ")", "", "" }, // ir_unop_f2i,      /**< Float-to-integer conversion. */
+	{ "float(", ")", "", "" }, // ir_unop_i2f,      /**< Integer-to-float conversion. */
+	{ "bool(", ")", "", "" }, // ir_unop_f2b,      /**< Float-to-boolean conversion */
+	{ "float(", ")", "", "" }, // ir_unop_b2f,      /**< Boolean-to-float conversion */
+	{ "bool(", ")", "", "" }, // ir_unop_i2b,      /**< int-to-boolean conversion */
+	{ "int(", ")", "", "" }, // ir_unop_b2i,      /**< Boolean-to-int conversion */
+	{ "uint(", ")", "", "" }, // ir_unop_b2u,
+	{ "bool(", ")", "", "" }, // ir_unop_u2b,
+	{ "uint(", ")", "", "" }, // ir_unop_f2u,
+	{ "float(", ")", "", "" }, // ir_unop_u2f,      /**< Unsigned-to-float conversion. */
+	{ "uint(", ")", "", "" }, // ir_unop_i2u,      /**< Integer-to-unsigned conversion. */
+	{ "int(", ")", "", "" }, // ir_unop_u2i,      /**< Unsigned-to-integer conversion. */
+	{ "int(", ")", "", "" }, // ir_unop_h2i,
+	{ "half(", ")", "", "" }, // ir_unop_i2h,
+	{ "float(", ")", "", "" }, // ir_unop_h2f,
+	{ "half(", ")", "", "" }, // ir_unop_f2h,
+	{ "bool(", ")", "", "" }, // ir_unop_h2b,
+	{ "float(", ")", "", "" }, // ir_unop_b2h,
+	{ "uint(", ")", "", "" }, // ir_unop_h2u,
+	{ "uint(", ")", "", "" }, // ir_unop_u2h,
+	{ "transpose(", ")", "", "" }, // ir_unop_transpose
+	{ "any(", ")", "", "" }, // ir_unop_any,
+	{ "all(", ")", "", "" }, // ir_unop_all,
+
+	/**
+	* \name Unary floating-point rounding operations.
+	*/
+	/*@{*/
+	{ "trunc(", ")", "", "" }, // ir_unop_trunc,
+	{ "ceil(", ")", "", "" }, // ir_unop_ceil,
+	{ "floor(", ")", "", "" }, // ir_unop_floor,
+	{ "fract(", ")", "", "" }, // ir_unop_fract,
+	{ "round(", ")", "", "" }, // ir_unop_round,
+	/*@}*/
+
+	/**
+	* \name Trigonometric operations.
+	*/
+	/*@{*/
+	{ "sin(", ")", "", "" }, // ir_unop_sin,
+	{ "cos(", ")", "", "" }, // ir_unop_cos,
+	{ "tan(", ")", "", "" }, // ir_unop_tan,
+	{ "asin(", ")", "", "" }, // ir_unop_asin,
+	{ "acos(", ")", "", "" }, // ir_unop_acos,
+	{ "atan(", ")", "", "" }, // ir_unop_atan,
+	{ "sinh(", ")", "", "" }, // ir_unop_sinh,
+	{ "cosh(", ")", "", "" }, // ir_unop_cosh,
+	{ "tanh(", ")", "", "" }, // ir_unop_tanh,
+	/*@}*/
+
+	/**
+	* \name Normalize.
+	*/
+	/*@{*/
+	{ "normalize(", ")", "", "" }, // ir_unop_normalize,
+	/*@}*/
+
+	/**
+	* \name Partial derivatives.
+	*/
+	/*@{*/
+	{ "dfdx(", ")", "", "" }, // ir_unop_dFdx,
+	{ "dfdy(", ")", "", "" }, // ir_unop_dFdy,
+	/*@}*/
+
+	{ "isnan(", ")", "", "" }, // ir_unop_isnan,
+	{ "isinf(", ")", "", "" }, // ir_unop_isinf,
+
+	{ "floatBitsToUint(", ")", "", "" }, // ir_unop_fasu,
+	{ "floatBitsToInt(", ")", "", "" }, // ir_unop_fasi,
+	{ "intBitsToFloat(", ")", "", "" }, // ir_unop_iasf,
+	{ "uintBitsToFloat(", ")", "", "" }, // ir_unop_uasf,
+
+	{ "bitfieldReverse(", ")", "", "" }, // ir_unop_bitreverse,
+	{ "bitCount(", ")", "", "" }, // ir_unop_bitcount,
+	{ "findMSB(", ")", "", "" }, // ir_unop_msb,
+	{ "findLSB(", ")", "", "" }, // ir_unop_lsb,
+
+	{ "ERROR_NO_NOISE_FUNCS(", ")", "", "" }, // ir_unop_noise,
+
+	{ "(", "+", ")", "" }, // ir_binop_add,
+	{ "(", "-", ")", "" }, // ir_binop_sub,
+	{ "(", "*", ")", "" }, // ir_binop_mul,
+	{ "(", "/", ")", "" }, // ir_binop_div,
+
+	/**
+	* Takes one of two combinations of arguments:
+	*
+	* - mod(vecN, vecN)
+	* - mod(vecN, float)
+	*
+	* Does not take integer types.
+	*/
+	{ "modf(", ",", ")", "%" }, // ir_binop_mod,
+	{ "modf(", ",", ")", "" }, // ir_binop_modf,
+
+	{ "step(", ",", ")", "" }, // ir_binop_step,
+
+	/**
+	* \name Binary comparison operators which return a boolean vector.
+	* The type of both operands must be equal.
+	*/
+	/*@{*/
+	{ "(", "<", ")", "<" }, // ir_binop_less,
+	{ "(", ">", ")", ">" }, // ir_binop_greater,
+	{ "(", "<=", ")", "<=" }, // ir_binop_lequal,
+	{ "(", ">=", ")", ">=" }, // ir_binop_gequal,
+	{ "(", "==", ")", "==" }, // ir_binop_equal,
+	{ "(", "!=", ")", "!=" }, // ir_binop_nequal,
+	/**
+	* Returns single boolean for whether all components of operands[0]
+	* equal the components of operands[1].
+	*/
+	{ "(", "==", ")", "" }, // ir_binop_all_equal,
+	/**
+	* Returns single boolean for whether any component of operands[0]
+	* is not equal to the corresponding component of operands[1].
+	*/
+	{ "(", "!=", ")", "" }, // ir_binop_any_nequal,
+	/*@}*/
+
+	/**
+	* \name Bit-wise binary operations.
+	*/
+	/*@{*/
+	{ "(", "<<", ")", "" }, // ir_binop_lshift,
+	{ "(", ">>", ")", "" }, // ir_binop_rshift,
+	{ "(", "&", ")", "" }, // ir_binop_bit_and,
+	{ "(", "^", ")", "" }, // ir_binop_bit_xor,
+	{ "(", "|", ")", "" }, // ir_binop_bit_or,
+	/*@}*/
+
+	{ "bvec%d(uvec%d(", ")*uvec%d(", "))", "&&" }, // ir_binop_logic_and,
+	{ "bvec%d(abs(ivec%d(", ")+ivec%d(", ")))", "^^" }, // ir_binop_logic_xor,
+	{ "bvec%d(uvec%d(", ")+uvec%d(", "))", "||" }, // ir_binop_logic_or,
+
+	{ "dot(", ",", ")", "" }, // ir_binop_dot,
+	{ "cross(", ",", ")", "" }, // ir_binop_cross,
+	{ "fmin(", ",", ")", "" }, // ir_binop_min,
+	{ "fmax(", ",", ")", "" }, // ir_binop_max,
+	{ "atan(", ",", ")", "" },
+	{ "pow(", ",", ")", "" }, // ir_binop_pow,
+
+	{ "mix(", ",", ",", ")" }, // ir_ternop_lerp,
+	{ "smoothstep(", ",", ",", ")" }, // ir_ternop_smoothstep,
+	{ "clamp(", ",", ",", ")" }, // ir_ternop_clamp,
+
+	{ "ERROR_QUADOP_VECTOR(", ",", ")" }, // ir_quadop_vector,
+};
+
+static_assert((sizeof(GLSLExpressionTable) / sizeof(GLSLExpressionTable[0])) == ir_opcode_count, "Metal Expression Table Size Mismatch");
+
+struct SDMARange
+{
+	unsigned SourceCB;
+	unsigned SourceOffset;
+	unsigned Size;
+	unsigned DestCBIndex;
+	unsigned DestCBPrecision;
+	unsigned DestOffset;
+
+	bool operator <(SDMARange const & Other) const
+	{
+		if (SourceCB == Other.SourceCB)
+		{
+			return SourceOffset < Other.SourceOffset;
+		}
+
+		return SourceCB < Other.SourceCB;
+	}
+};
+typedef std::list<SDMARange> TDMARangeList;
+typedef std::map<unsigned, TDMARangeList> TCBDMARangeMap;
+
+
+static void InsertRange( TCBDMARangeMap& CBAllRanges, unsigned SourceCB, unsigned SourceOffset, unsigned Size, unsigned DestCBIndex, unsigned DestCBPrecision, unsigned DestOffset ) 
+{
+	check(SourceCB < (1 << 12));
+	check(DestCBIndex < (1 << 12));
+	check(DestCBPrecision < (1 << 8));
+	unsigned SourceDestCBKey = (SourceCB << 20) | (DestCBIndex << 8) | DestCBPrecision;
+	SDMARange Range = { SourceCB, SourceOffset, Size, DestCBIndex, DestCBPrecision, DestOffset };
+
+	TDMARangeList& CBRanges = CBAllRanges[SourceDestCBKey];
+//printf("* InsertRange: %08x\t%d:%d - %d:%c:%d:%d\n", SourceDestCBKey, SourceCB, SourceOffset, DestCBIndex, DestCBPrecision, DestOffset, Size);
+	if (CBRanges.empty())
+	{
+		CBRanges.push_back(Range);
+	}
+	else
+	{
+		TDMARangeList::iterator Prev = CBRanges.end();
+		bool bAdded = false;
+		for (auto Iter = CBRanges.begin(); Iter != CBRanges.end(); ++Iter)
+		{
+			if (SourceOffset + Size <= Iter->SourceOffset)
+			{
+				if (Prev == CBRanges.end())
+				{
+					CBRanges.push_front(Range);
+				}
+				else
+				{
+					CBRanges.insert(Iter, Range);
+				}
+
+				bAdded = true;
+				break;
+			}
+
+			Prev = Iter;
+		}
+
+		if (!bAdded)
+		{
+			CBRanges.push_back(Range);
+		}
+
+		if (CBRanges.size() > 1)
+		{
+			// Try to merge ranges
+			bool bDirty = false;
+			do
+			{
+				bDirty = false;
+				TDMARangeList NewCBRanges;
+				for (auto Iter = CBRanges.begin(); Iter != CBRanges.end(); ++Iter)
+				{
+					if (Iter == CBRanges.begin())
+					{
+						Prev = CBRanges.begin();
+					}
+					else
+					{
+						if (Prev->SourceOffset + Prev->Size == Iter->SourceOffset && Prev->DestOffset + Prev->Size == Iter->DestOffset)
+						{
+							SDMARange Merged = *Prev;
+							Merged.Size = Prev->Size + Iter->Size;
+							NewCBRanges.pop_back();
+							NewCBRanges.push_back(Merged);
+							++Iter;
+							NewCBRanges.insert(NewCBRanges.end(), Iter, CBRanges.end());
+							bDirty = true;
+							break;
+						}
+					}
+
+					NewCBRanges.push_back(*Iter);
+					Prev = Iter;
+				}
+
+				CBRanges.swap(NewCBRanges);
+			}
+			while (bDirty);
+		}
+	}
+}
+
+static TDMARangeList SortRanges( TCBDMARangeMap& CBRanges ) 
+{
+	TDMARangeList Sorted;
+	for (auto Iter = CBRanges.begin(); Iter != CBRanges.end(); ++Iter)
+	{
+		Sorted.insert(Sorted.end(), Iter->second.begin(), Iter->second.end());
+	}
+
+	Sorted.sort();
+
+	return Sorted;
+}
+
+static void DumpSortedRanges(TDMARangeList& SortedRanges)
+{
+	printf("**********************************\n");
+	for (auto i = SortedRanges.begin(); i != SortedRanges.end(); ++i )
+	{
+		auto o = *i;
+		printf("\t%d:%d - %d:%c:%d:%d\n", o.SourceCB, o.SourceOffset, o.DestCBIndex, o.DestCBPrecision, o.DestOffset, o.Size);
+	}
+}
+
+/**
+ * IR visitor used to generate Metal. Based on ir_print_visitor.
+ */
+class FGenerateMetalVisitor : public ir_visitor
+{
+	_mesa_glsl_parse_state* ParseState;
+
+	/** Track which multi-dimensional arrays are used. */
+	struct md_array_entry : public exec_node
+	{
+		const glsl_type* type;
+	};
+
+public:
+	/** External variables. */
+	exec_list input_variables;
+protected:
+	exec_list output_variables;
+	exec_list uniform_variables;
+	exec_list sampler_variables;
+	exec_list image_variables;
+
+	/** Data tied globally to the shader via attributes */
+	int wg_size_x;
+	int wg_size_y;
+	int wg_size_z;
+
+	glsl_tessellation_info tessellation;
+
+	/** Track global instructions. */
+	struct global_ir : public exec_node
+	{
+		ir_instruction* ir;
+		explicit global_ir(ir_instruction* in_ir) : ir(in_ir) {}
+	};
+
+	/** Global instructions. */
+	exec_list global_instructions;
+
+	/** A mapping from ir_variable * -> unique printable names. */
+	hash_table *printable_names;
+	/** Structures required by the code. */
+	hash_table *used_structures;
+	/** Uniform block variables required by the code. */
+	hash_table *used_uniform_blocks;
+	/** Multi-dimensional arrays required by the code. */
+	exec_list used_md_arrays;
+
+	// Code generation flags
+	bool bIsVS;
+
+	FBuffers& Buffers;
+
+	/** Memory context within which to make allocations. */
+	void *mem_ctx;
+	/** Buffer to which GLSL source is being generated. */
+	char** buffer;
+	/** Indentation level. */
+	int indentation;
+	/** Scope depth. */
+	int scope_depth;
+	/** The number of temporary variables declared in the current scope. */
+	int temp_id;
+	/** The number of global variables declared. */
+	int global_id;
+	/** Whether a semicolon must be printed before the next EOL. */
+	bool needs_semicolon;
+	bool IsMain;
+	/**
+	 * Whether uint literals should be printed as int literals. This is a hack
+	 * because glCompileShader crashes on Mac OS X with code like this:
+	 * foo = bar[0u];
+	 */
+	bool should_print_uint_literals_as_ints;
+	/** number of loops in the generated code */
+	int loop_count;
+
+	// Only one stage_in is allowed
+	bool bStageInEmitted;
+
+	// Are we in the middle of a function signature?
+	//bool bIsFunctionSig = false;
+
+	// Use packed_ prefix when printing out structs
+	bool bUsePacked;
+
+	/**
+	 * Return true if the type is a multi-dimensional array. Also, track the
+	 * array.
+	 */
+	bool is_md_array(const glsl_type* type)
+	{
+		if (type->base_type == GLSL_TYPE_ARRAY &&
+			type->fields.array->base_type == GLSL_TYPE_ARRAY)
+		{
+			foreach_iter(exec_list_iterator, iter, used_md_arrays) 
+			{
+				md_array_entry* entry = (md_array_entry*)iter.get();
+				if (entry->type == type)
+					return true;
+				}
+			md_array_entry* entry = new(mem_ctx) md_array_entry();
+			entry->type = type;
+			used_md_arrays.push_tail(entry);
+			return true;
+		}
+		return false;
+	}
+
+	/**
+	 * Fetch/generate a unique name for ir_variable.
+	 *
+	 * GLSL IR permits multiple ir_variables to share the same name.  This works
+	 * fine until we try to print it, when we really need a unique one.
+	 */
+	const char *unique_name(ir_variable *var)
+	{
+		if (var->mode == ir_var_temporary || var->mode == ir_var_auto)
+		{
+			/* Do we already have a name for this variable? */
+			const char *name = (const char *) hash_table_find(this->printable_names, var);
+			if (name == NULL)
+			{
+				bool is_global = (scope_depth == 0 && var->mode != ir_var_temporary);
+				const char *prefix = is_global ? "g" : "t";
+				int var_id = is_global ? global_id++ : temp_id++;
+				name = ralloc_asprintf(mem_ctx, "%s%d", prefix, var_id);
+				hash_table_insert(this->printable_names, (void *)name, var);
+			}
+			return name;
+		}
+
+		/* If there's no conflict, just use the original name */
+		return var->name;
+	}
+
+	/**
+	 * Add tabs/spaces for the current indentation level.
+	 */
+	void indent(void)
+	{
+		for (int i = 0; i < indentation; i++)
+		{
+			ralloc_asprintf_append(buffer, "\t");
+		}
+	}
+
+	/**
+	 * Print out the internal name for a multi-dimensional array.
+	 */
+	void print_md_array_type(const glsl_type *t)
+	{
+		if (t->base_type == GLSL_TYPE_ARRAY)
+		{
+			ralloc_asprintf_append(buffer, "_mdarr_");
+			do 
+			{
+				ralloc_asprintf_append(buffer, "%d_", t->length);
+				t = t->fields.array;
+			} while (t->base_type == GLSL_TYPE_ARRAY);
+			print_base_type(t);
+		}
+	}
+
+	/**
+	 * Print the base type, e.g. vec3.
+	 */
+	void print_base_type(const glsl_type *t)
+	{
+		if (t->base_type == GLSL_TYPE_ARRAY)
+		{
+			print_base_type(t->fields.array);
+		}
+		else if (t->base_type == GLSL_TYPE_INPUTPATCH)
+		{
+			ralloc_asprintf_append(buffer, "GLSL_TYPE_INPUTPATCH");
+			check(0);
+			ralloc_asprintf_append(buffer, "/* %s */ ", t->name);
+			print_base_type(t->inner_type);
+		}
+		else if (t->base_type == GLSL_TYPE_OUTPUTPATCH)
+		{
+			ralloc_asprintf_append(buffer, "GLSL_TYPE_OUTPUTPATCH");
+			check(0);
+			ralloc_asprintf_append(buffer, "/* %s */ ", t->name);
+			print_base_type(t->inner_type);
+		}
+		else if ((t->base_type == GLSL_TYPE_STRUCT)
+				&& (strncmp("gl_", t->name, 3) != 0))
+		{
+			ralloc_asprintf_append(buffer, "%s", t->name);
+		}
+		else 
+		{
+			check(t->HlslName);
+			if (/*!bIsFunctionSig &&*/ bUsePacked && t->is_vector() && t->vector_elements < 4)
+			{
+				ralloc_asprintf_append(buffer, "packed_%s", t->HlslName);
+			}
+			else
+			{
+				ralloc_asprintf_append(buffer, "%s", t->HlslName);
+			}
+		}
+	}
+
+	/**
+	 * Print the portion of the type that appears before a variable declaration.
+	 */
+	void print_type_pre(const glsl_type *t)
+	{
+		if (is_md_array(t))
+		{
+			print_md_array_type(t);
+		}
+		else
+		{
+			print_base_type(t);
+		}
+	}
+
+	/**
+	 * Print the portion of the type that appears after a variable declaration.
+	 */
+	void print_type_post(const glsl_type *t)
+	{
+		if (t->base_type == GLSL_TYPE_ARRAY && !is_md_array(t))
+		{
+			ralloc_asprintf_append(buffer, "[%u]", t->length);
+		}
+		else if (t->base_type == GLSL_TYPE_INPUTPATCH || t->base_type == GLSL_TYPE_OUTPUTPATCH)
+		{
+			ralloc_asprintf_append(buffer, "[%u] /* %s */", t->patch_length, t->name);
+		}
+	}
+
+	/**
+	 * Print a full variable declaration.
+	 */
+	void print_type_full(const glsl_type *t)
+	{
+		print_type_pre(t);
+		print_type_post(t);
+	}
+
+	/**
+	 * Visit a single instruction. Appends a semicolon and EOL if needed.
+	 */
+	void do_visit(ir_instruction* ir)
+	{
+		needs_semicolon = true;
+		ir->accept(this);
+		if (needs_semicolon)
+		{
+			ralloc_asprintf_append(buffer, ";\n");
+		}
+	}
+
+	/**
+	* \name Visit methods
+	*
+	* As typical for the visitor pattern, there must be one \c visit method for
+	* each concrete subclass of \c ir_instruction.  Virtual base classes within
+	* the hierarchy should not have \c visit methods.
+	*/
+
+	virtual void visit(ir_rvalue *rvalue) override
+	{
+		check(0 && "ir_rvalue not handled for GLSL export.");
+	}
+
+	virtual void visit(ir_variable *var) override
+	{
+		// Check for an initialized const variable
+		// If var is read-only and initialized, set it up as an initialized const
+		bool constInit = false;
+		if (var->has_initializer && var->read_only && (var->constant_initializer || var->constant_value))
+		{
+			ralloc_asprintf_append( buffer, "const ");
+			constInit = true;
+		}
+
+		if (scope_depth == 0)
+		{
+			check(0);
+		}
+
+		if (scope_depth == 0 && var->mode == ir_var_temporary)
+		{
+			check(0);
+		}
+		else
+		{
+			int layout_bits =
+				(var->origin_upper_left ? 0x1 : 0) |
+				(var->pixel_center_integer ? 0x2 : 0);
+
+			if (scope_depth == 0 &&
+			   ((var->mode == ir_var_in) || (var->mode == ir_var_out)) && 
+			   var->is_interface_block)
+			{
+				check(0);
+			}
+			else if (var->type->is_image())
+			{
+				check(0);
+			}
+			else
+			{
+				if (IsMain && var->type->base_type == GLSL_TYPE_STRUCT && (var->mode == ir_var_in || var->mode == ir_var_out || var->mode == ir_var_uniform))
+				{
+					if (hash_table_find(used_structures, var->type) == NULL)
+					{
+						hash_table_insert(used_structures, (void*)var->type, var->type);
+					}
+				}
+
+				if (IsMain && var->mode == ir_var_uniform)
+				{
+					auto* PtrType = var->type->is_array() ? var->type->element_type() : var->type;
+					check(!PtrType->is_array());
+
+					if (var->type->is_sampler())
+					{
+						auto* Entry = ParseState->FindPackedSamplerEntry(var->name);
+						check(Entry);
+						//@todo-rco: SamplerStates
+						auto SamplerStateFound = ParseState->TextureToSamplerMap.find(Entry->Name.c_str());
+						if (SamplerStateFound != ParseState->TextureToSamplerMap.end())
+						{
+							auto& SamplerStates = SamplerStateFound->second;
+							if (SamplerStates.size() == 1)
+							{
+								ralloc_asprintf_append(
+									buffer,
+									"sampler s%d [[ sampler(%d) ]], ", Entry->offset, Entry->offset);
+							}
+							else
+							{
+								check(SamplerStates.empty());
+							}
+						}
+
+						print_type_pre(PtrType);
+						ralloc_asprintf_append(
+							buffer,
+							"<%s> %s", (PtrType->inner_type && PtrType->inner_type->base_type == GLSL_TYPE_HALF) ? "half" : "float", unique_name(var));
+						print_type_post(PtrType);
+						ralloc_asprintf_append(
+							buffer,
+							" [[ texture(%d) ]]", Entry->offset
+							);
+					}
+					else
+					{
+						int BufferIndex = Buffers.GetIndex(var);
+						bool bNeedsPointer = var->semantic && (strlen(var->semantic) == 1);
+						check(BufferIndex >= 0);
+						ralloc_asprintf_append(
+							buffer,
+							"constant "
+							);
+						print_type_pre(PtrType);
+						ralloc_asprintf_append(buffer, " %s%s", bNeedsPointer ? "*" : "&", unique_name(var));
+						print_type_post(PtrType);
+						ralloc_asprintf_append(
+							buffer,
+							" [[ buffer(%d) ]]", BufferIndex
+							);
+					}
+				}
+				else if (IsMain && var->mode == ir_var_in)
+				{
+					if (!strcmp(var->name, "gl_FrontFacing"))
+					{
+						check(var->type->is_boolean());
+						print_type_pre(var->type);
+						ralloc_asprintf_append(buffer, " %s", unique_name(var));
+						print_type_post(var->type);
+						ralloc_asprintf_append(
+							buffer,
+							" [[ front_facing ]]"
+							);
+					}
+					else if (var->semantic && !strncmp(var->semantic, "[[ color(", 9))
+					{
+						check(var->type->is_vector() && var->type->vector_elements == 4);
+						print_type_pre(var->type);
+						ralloc_asprintf_append(buffer, " %s", unique_name(var));
+						print_type_post(var->type);
+						ralloc_asprintf_append(
+							buffer,
+							" %s",
+							var->semantic
+							);
+					}
+					else if(var->semantic && !strncmp(var->semantic,"[[", 2))
+					{
+						check(!var->type->is_record());
+						print_type_pre(var->type);
+						ralloc_asprintf_append(buffer," %s",unique_name(var));
+						print_type_post(var->type);
+						ralloc_asprintf_append(
+							buffer,
+							" %s",
+							var->semantic
+							);
+					}
+					else
+					{
+						check(var->type->is_record());
+						check(!bStageInEmitted);
+						print_type_pre(var->type);
+						ralloc_asprintf_append(buffer, " %s", unique_name(var));
+						print_type_post(var->type);
+						ralloc_asprintf_append(
+							buffer,
+							" [[ stage_in ]]"
+							);
+						bStageInEmitted = true;
+					}
+				}
+				else if (IsMain && var->mode == ir_var_out)
+				{
+					auto* PtrType = var->type->is_array() ? var->type->element_type() : var->type;
+					check(!PtrType->is_array());
+					print_type_pre(PtrType);
+					ralloc_asprintf_append(buffer, " %s", unique_name(var));
+					print_type_post(PtrType);
+				}
+				else
+				{
+					print_type_pre(var->type);
+					ralloc_asprintf_append(buffer, " %s", unique_name(var));
+					print_type_post(var->type);
+				}
+			}
+		}
+
+		// Add the initializer if we need it
+		if (constInit)
+		{
+			ralloc_asprintf_append(buffer, " = ");
+			if (var->constant_initializer)
+			{
+				var->constant_initializer->accept(this);
+			}
+			else
+			{
+				var->constant_value->accept(this);
+			}
+		}
+
+	}
+
+	virtual void visit(ir_function_signature *sig) override
+	{
+		// Reset temporary id count.
+		temp_id = 0;
+		bool bPrintComma = false;
+		scope_depth++;
+		IsMain = sig->is_main;
+
+		print_type_full(sig->return_type);
+		ralloc_asprintf_append(buffer, " %s(", sig->function_name());
+
+		foreach_iter(exec_list_iterator, iter, sig->parameters)
+		{
+			ir_variable *const inst = (ir_variable *) iter.get();
+			if (bPrintComma)
+			{
+				ralloc_asprintf_append(buffer, ",\n");
+				++indentation;
+				indent();
+				--indentation;
+			}
+			inst->accept(this);
+			bPrintComma = true;
+		}
+		ralloc_asprintf_append(buffer, ")\n");
+
+		indent();
+		ralloc_asprintf_append(buffer, "{\n");
+
+		if (sig->is_main && !global_instructions.is_empty())
+		{
+			indentation++;
+			foreach_iter(exec_list_iterator, iter, global_instructions)
+			{
+				global_ir* gir = (global_ir*)iter.get();
+				indent();
+				do_visit(gir->ir);
+			}
+			indentation--;
+		}
+
+		//grab the global attributes
+		if (sig->is_main)
+		{
+			wg_size_x = sig->wg_size_x;
+			wg_size_y = sig->wg_size_y;
+			wg_size_z = sig->wg_size_z;
+
+			tessellation = sig->tessellation;
+		}
+
+		indentation++;
+		foreach_iter(exec_list_iterator, iter, sig->body)
+		{
+			ir_instruction *const inst = (ir_instruction *) iter.get();
+			indent();
+			do_visit(inst);
+		}
+		indentation--;
+		indent();
+		ralloc_asprintf_append(buffer, "}\n");
+		needs_semicolon = false;
+		IsMain = false;
+		scope_depth--;
+	}
+
+	virtual void visit(ir_function *func) override
+	{
+		foreach_iter(exec_list_iterator, iter, *func)
+		{
+			ir_function_signature *const sig = (ir_function_signature *) iter.get();
+			if (sig->is_defined && !sig->is_builtin)
+			{
+				indent();
+				if (sig->is_main)
+				{
+					ralloc_asprintf_append(buffer, "%s ", bIsVS ? "vertex" : "fragment");
+				}
+
+				sig->accept(this);
+			}
+		}
+		needs_semicolon = false;
+	}
+
+	virtual void visit(ir_expression *expr) override
+	{
+		check(scope_depth > 0);
+
+		int numOps = expr->get_num_operands();
+		ir_expression_operation op = expr->operation;
+
+		if (op == ir_unop_rcp)
+		{
+			check(numOps == 1);
+
+			std::string Type = expr->type->name;
+			Type = FixVecPrefix(Type);
+
+			ralloc_asprintf_append(buffer, "(%s(1.0) / ", Type.c_str());
+			expr->operands[0]->accept(this);
+			ralloc_asprintf_append(buffer, ")");
+		}
+		else if (numOps == 1 && op >= ir_unop_first_conversion && op <= ir_unop_last_conversion)
+		{
+			std::string Type = expr->type->name;
+			Type = FixVecPrefix(Type);
+
+			ralloc_asprintf_append(buffer, "%s(", Type.c_str());
+			expr->operands[0]->accept(this);
+			ralloc_asprintf_append(buffer, ")");
+		}
+		else if (expr->type->is_scalar() &&
+			((numOps == 1 && op == ir_unop_logic_not) ||
+			 (numOps == 2 && op >= ir_binop_first_comparison && op <= ir_binop_last_comparison) ||
+			 (numOps == 2 && op >= ir_binop_first_logic && op <= ir_binop_last_logic)))
+		{
+			const char* op_str = GLSLExpressionTable[op][3];
+			ralloc_asprintf_append(buffer, "%s(", (numOps == 1) ? op_str : "");
+			expr->operands[0]->accept(this);
+			if (numOps == 2)
+			{
+				ralloc_asprintf_append(buffer, "%s", op_str);
+				expr->operands[1]->accept(this);
+			}
+			ralloc_asprintf_append(buffer, ")");
+		}
+		else if (expr->type->is_vector() && numOps == 2 &&
+			op >= ir_binop_first_logic && op <= ir_binop_last_logic)
+		{
+			ralloc_asprintf_append(buffer, GLSLExpressionTable[op][0], expr->type->vector_elements, expr->type->vector_elements);
+			expr->operands[0]->accept(this);
+			ralloc_asprintf_append(buffer, GLSLExpressionTable[op][1], expr->type->vector_elements);
+			expr->operands[1]->accept(this);
+			ralloc_asprintf_append(buffer, GLSLExpressionTable[op][2]);
+		}
+		else if (op == ir_binop_mod && !expr->type->is_float())
+		{
+			ralloc_asprintf_append(buffer, "((");
+			expr->operands[0]->accept(this);
+			ralloc_asprintf_append(buffer, ")%%(");
+			expr->operands[1]->accept(this);
+			ralloc_asprintf_append(buffer, "))");
+		}
+		else if (op == ir_binop_mul && expr->type->is_matrix()
+			&& expr->operands[0]->type->is_matrix()
+			&& expr->operands[1]->type->is_matrix())
+		{
+			ralloc_asprintf_append(buffer, "ERRROR_MulMatrix()");
+			check(0);
+		}
+		else if (numOps < 4)
+		{
+			ralloc_asprintf_append(buffer, GLSLExpressionTable[op][0]);
+			for (int i = 0; i < numOps; ++i)
+			{
+				expr->operands[i]->accept(this);
+				ralloc_asprintf_append(buffer, GLSLExpressionTable[op][i+1]);
+			}
+		}
+	}
+
+	virtual void visit(ir_texture *tex) override
+	{
+		check(scope_depth > 0);
+
+		tex->sampler->accept(this);
+		if (tex->op == ir_tex || tex->op == ir_txl || tex->op == ir_txb)
+		{
+			ralloc_asprintf_append(buffer, ".sample(");
+			auto* Texture = tex->sampler->variable_referenced();
+			check(Texture);
+			auto* Entry = ParseState->FindPackedSamplerEntry(Texture->name);
+			ralloc_asprintf_append(buffer, "s%d, ", Entry->offset);
+			tex->coordinate->accept(this);
+
+			if (tex->op == ir_txl)
+			{
+				ralloc_asprintf_append(buffer, ", level(");
+				tex->lod_info.lod->accept(this);
+				ralloc_asprintf_append(buffer, ")");
+			}
+			else if (tex->op == ir_txb)
+			{
+				ralloc_asprintf_append(buffer, ", bias(");
+				tex->lod_info.lod->accept(this);
+				ralloc_asprintf_append(buffer, ")");
+			}
+
+			if (tex->offset)
+			{
+				ralloc_asprintf_append(buffer, ", ");
+				tex->offset->accept(this);
+			}
+		}
+		else if (tex->op == ir_txf)
+		{
+			ralloc_asprintf_append(buffer, ".read(");
+			tex->coordinate->accept(this);
+		}
+		else
+		{
+			ralloc_asprintf_append(buffer, "UNKNOWN TEXOP %d!", tex->op);
+			check(0);
+		}
+
+		ralloc_asprintf_append(buffer, ")");
+	}
+
+	virtual void visit(ir_swizzle *swizzle) override
+	{
+		check(scope_depth > 0);
+
+		const unsigned mask[4] =
+		{
+			swizzle->mask.x,
+			swizzle->mask.y,
+			swizzle->mask.z,
+			swizzle->mask.w,
+		};
+
+		if (swizzle->val->type->is_scalar())
+		{
+			// Scalar -> Vector swizzles must use the constructor syntax.
+			if (swizzle->type->is_scalar() == false)
+			{
+				print_type_full(swizzle->type);
+				ralloc_asprintf_append(buffer, "(");
+				swizzle->val->accept(this);
+				ralloc_asprintf_append(buffer, ")");
+			}
+		}
+		else
+		{
+			const bool is_constant = swizzle->val->as_constant() != nullptr;
+			if (is_constant)
+			{
+				ralloc_asprintf_append(buffer, "(");
+			}
+			swizzle->val->accept(this);
+			if (is_constant)
+			{
+				ralloc_asprintf_append(buffer, ")");
+			}
+			ralloc_asprintf_append(buffer, ".");
+			for (unsigned i = 0; i < swizzle->mask.num_components; i++)
+			{
+				ralloc_asprintf_append(buffer, "%c", "xyzw"[mask[i]]);
+			}
+		}
+	}
+
+	virtual void visit(ir_dereference_variable *deref) override
+	{
+		check(scope_depth > 0);
+
+		ir_variable* var = deref->variable_referenced();
+
+		ralloc_asprintf_append(buffer, unique_name(var));
+
+		if (var->type->base_type == GLSL_TYPE_STRUCT)
+		{
+			if (hash_table_find(used_structures, var->type) == NULL)
+			{
+				hash_table_insert(used_structures, (void*)var->type, var->type);
+			}
+		}
+
+		if (var->type->base_type == GLSL_TYPE_ARRAY && var->type->fields.array->base_type == GLSL_TYPE_STRUCT)
+		{
+			if (hash_table_find(used_structures, var->type->fields.array) == NULL)
+			{
+				hash_table_insert(used_structures, (void*)var->type->fields.array, var->type->fields.array);
+			}
+		}
+
+		if ((var->type->base_type == GLSL_TYPE_INPUTPATCH || var->type->base_type == GLSL_TYPE_OUTPUTPATCH) && var->type->inner_type->base_type == GLSL_TYPE_STRUCT)
+		{
+			if (hash_table_find(used_structures, var->type->inner_type) == NULL)
+			{
+				hash_table_insert(used_structures, (void*)var->type->inner_type, var->type->inner_type);
+			}
+		}
+
+		if (var->mode == ir_var_uniform && var->semantic != NULL)
+		{
+			if (hash_table_find(used_uniform_blocks, var->semantic) == NULL)
+			{
+				hash_table_insert(used_uniform_blocks, (void*)var->semantic, var->semantic);
+			}
+		}
+
+		if (is_md_array(deref->type))
+		{
+			ralloc_asprintf_append(buffer, ".Inner");
+		}
+	}
+
+	virtual void visit(ir_dereference_array *deref) override
+	{
+		check(scope_depth > 0);
+
+		deref->array->accept(this);
+
+		// Make extra sure crappy Mac OS X compiler won't have any reason to crash
+		bool enforceInt = false;
+
+		if (deref->array_index->type->base_type == GLSL_TYPE_UINT)
+		{
+			if (deref->array_index->ir_type == ir_type_constant)
+			{
+				should_print_uint_literals_as_ints = true;
+			}
+			else
+			{
+				enforceInt = true;
+			}
+		}
+
+		if (enforceInt)
+		{
+			ralloc_asprintf_append(buffer, "[int(");
+		}
+		else
+		{
+			ralloc_asprintf_append(buffer, "[");
+		}
+
+		deref->array_index->accept(this);
+		should_print_uint_literals_as_ints = false;
+
+		if (enforceInt)
+		{
+			ralloc_asprintf_append(buffer, ")]");
+		}
+		else
+		{
+			ralloc_asprintf_append(buffer, "]");
+		}
+
+		if (is_md_array(deref->array->type))
+		{
+			ralloc_asprintf_append(buffer, ".Inner");
+		}
+	}
+
+	void print_image_op( ir_dereference_image *deref, ir_rvalue *src)
+	{
+		const char* swizzle[] =
+		{
+			"x", "xy", "xyz", "xyzw"
+		};
+		const char* expand[] =
+		{
+			"xxxx", "xyxx", "xyzx", "xyzw"
+		};
+		const char* int_cast[] =
+		{
+			"int", "ivec2", "ivec3", "ivec4"
+		};
+		const int dst_elements = deref->type->vector_elements;
+		const int src_elements = (src) ? src->type->vector_elements : 1;
+		
+		check( 1 <= dst_elements && dst_elements <= 4);
+		check( 1 <= src_elements && src_elements <= 4);
+
+		if ( deref->op == ir_image_access)
+		{
+			if ( src == NULL )
+			{
+				ralloc_asprintf_append( buffer, "imageLoad( " );
+				deref->image->accept(this);
+				ralloc_asprintf_append(buffer, ", ");
+				deref->image_index->accept(this);
+				ralloc_asprintf_append(buffer, ").%s", swizzle[dst_elements-1]);
+			}
+			else
+			{
+				ralloc_asprintf_append( buffer, "imageStore( " );
+				deref->image->accept(this);
+				ralloc_asprintf_append(buffer, ", ");
+				deref->image_index->accept(this);
+				ralloc_asprintf_append(buffer, ", ");
+				src->accept(this);
+				ralloc_asprintf_append(buffer, ".%s)", expand[src_elements-1]);
+			}
+		}
+		else if ( deref->op == ir_image_dimensions)
+		{
+			ralloc_asprintf_append( buffer, "imageSize( " );
+			deref->image->accept(this);
+			ralloc_asprintf_append(buffer, ")");
+		}
+		else
+		{
+			check( !"Unknown image operation");
+		}
+	}
+
+	virtual void visit(ir_dereference_image *deref) override
+	{
+		check(scope_depth > 0);
+
+		print_image_op( deref, NULL);
+		
+	}
+
+	virtual void visit(ir_dereference_record *deref) override
+	{
+		check(scope_depth > 0);
+
+		deref->record->accept(this);
+		ralloc_asprintf_append(buffer, ".%s", deref->field);
+
+		if (is_md_array(deref->type))
+		{
+			ralloc_asprintf_append(buffer, ".Inner");
+		}
+	}
+
+	virtual void visit(ir_assignment *assign) override
+	{
+		if (scope_depth == 0)
+		{
+			global_instructions.push_tail(new(mem_ctx) global_ir(assign));
+			needs_semicolon = false;
+			return;
+		}
+
+		// constant variables with initializers are statically assigned
+		ir_variable *var = assign->lhs->variable_referenced();
+		if (var->has_initializer && var->read_only && (var->constant_initializer || var->constant_value))
+		{
+			//This will leave a blank line with a semi-colon
+			return;
+		}
+
+		if (assign->condition)
+		{
+			ralloc_asprintf_append(buffer, "if(");
+			assign->condition->accept(this);
+			ralloc_asprintf_append(buffer, ") { ");
+		}
+
+		if (assign->lhs->as_dereference_image() != NULL)
+		{
+			/** EHart - should the write mask be checked here? */
+			print_image_op( assign->lhs->as_dereference_image(), assign->rhs);
+		}
+		else
+		{
+			char mask[6];
+			unsigned j = 1;
+			if (assign->lhs->type->is_scalar() == false ||
+				assign->write_mask != 0x1)
+			{
+				for (unsigned i = 0; i < 4; i++)
+				{
+					if ((assign->write_mask & (1 << i)) != 0)
+					{
+						mask[j] = "xyzw"[i];
+						j++;
+					}
+				}
+			}
+			mask[j] = '\0';
+
+			mask[0] = (j == 1) ? '\0' : '.';
+
+			assign->lhs->accept(this);
+			ralloc_asprintf_append(buffer, "%s = ", mask);
+
+			// Hack: Need to add additional cast from packed types
+			bool bNeedToAcceptRHS = true;
+			if (auto* Expr = assign->rhs->as_expression())
+			{
+				if (Expr->operation == ir_unop_f2h)
+				{
+					auto* Var = Expr->operands[0]->variable_referenced();
+					if (Var && Var->mode == ir_var_uniform && Var->type->HlslName && !strcmp(Var->type->HlslName, "__PACKED__"))
+					{
+						ralloc_asprintf_append(buffer, "%s(%s(", Expr->type->name, FixVecPrefix(PromoteHalfToFloatType(ParseState, Expr->type)->name).c_str());
+						Expr->operands[0]->accept(this);
+						ralloc_asprintf_append(buffer, "))");
+						bNeedToAcceptRHS = false;
+					}
+				}
+			}
+
+			if (bNeedToAcceptRHS)
+			{
+				assign->rhs->accept(this);
+			}
+		}
+
+		if (assign->condition)
+		{
+			ralloc_asprintf_append(buffer, "%s }", needs_semicolon ? ";" : "");
+		}
+	}
+
+	void print_constant(ir_constant *constant, int index)
+	{
+		if (constant->type->is_float())
+		{
+			if (constant->is_component_finite(index))
+			{
+				float value = constant->value.f[index];
+				const char *format = (fabsf(fmodf(value,1.0f)) < 1.e-8f) ? "%.1f" : "%.8f";
+				ralloc_asprintf_append(buffer, format, value);
+			}
+			else
+			{
+				switch (constant->value.u[index])
+				{
+					case 0x7f800000u:
+						ralloc_asprintf_append(buffer, "(1.0/0.0)");
+						break;
+
+					case 0xffc00000u:
+						ralloc_asprintf_append(buffer, "(0.0/0.0)");
+						break;
+
+					case 0xff800000u:
+						ralloc_asprintf_append(buffer, "(-1.0/0.0)");
+						break;
+
+					default:
+						ralloc_asprintf_append(buffer, "UnknownNonFinite_0x%08x", constant->value.u[index]);
+						check(0);
+				}
+			}
+		}
+		else if (constant->type->base_type == GLSL_TYPE_INT)
+		{
+			ralloc_asprintf_append(buffer, "%d", constant->value.i[index]);
+		}
+		else if (constant->type->base_type == GLSL_TYPE_UINT)
+		{
+			ralloc_asprintf_append(buffer, "%u%s",
+				constant->value.u[index],
+				should_print_uint_literals_as_ints ? "" : "u"
+				);
+		}
+		else if (constant->type->base_type == GLSL_TYPE_BOOL)
+		{
+			ralloc_asprintf_append(buffer, "%s", constant->value.b[index] ? "true" : "false");
+		}
+	}
+
+	virtual void visit(ir_constant *constant) override
+	{
+		if (constant->type == glsl_type::float_type
+			|| constant->type == glsl_type::int_type
+			|| constant->type == glsl_type::uint_type)
+		{
+			print_constant(constant, 0);
+		}
+		else if (constant->type->is_record())
+		{
+			print_type_full(constant->type);
+			ralloc_asprintf_append(buffer, "(");
+			ir_constant* value = (ir_constant*)constant->components.get_head();
+			if (value)
+			{
+				value->accept(this);
+			}
+			for (uint32 i = 1; i < constant->type->length; i++)
+			{
+				check(value);
+				value = (ir_constant*)value->next;
+				if (value)
+				{
+					ralloc_asprintf_append(buffer, ",");
+					value->accept(this);
+				}
+			}
+			ralloc_asprintf_append(buffer, ")");
+		}
+		else if (constant->type->is_array())
+		{
+			print_type_full(constant->type);
+			ralloc_asprintf_append(buffer, "(");
+			constant->get_array_element(0)->accept(this);
+			for (uint32 i = 1; i < constant->type->length; ++i)
+			{
+				ralloc_asprintf_append(buffer, ",");
+				constant->get_array_element(i)->accept(this);
+			}
+			ralloc_asprintf_append(buffer, ")");
+		}
+		else
+		{
+			print_type_full(constant->type);
+			ralloc_asprintf_append(buffer, "(");
+			print_constant(constant, 0);
+			int num_components = constant->type->components();
+			for (int i = 1; i < num_components; ++i)
+			{
+				ralloc_asprintf_append(buffer, ",");
+				print_constant(constant, i);
+			}
+			ralloc_asprintf_append(buffer, ")");
+		}
+	}
+
+	virtual void visit(ir_call *call) override
+	{
+		if (scope_depth == 0)
+		{
+			global_instructions.push_tail(new(mem_ctx) global_ir(call));
+			needs_semicolon = false;
+			return;
+		}
+
+		if (call->return_deref)
+		{
+			call->return_deref->accept(this);
+			ralloc_asprintf_append(buffer, " = ");
+		}
+		ralloc_asprintf_append(buffer, "%s(", call->callee_name());
+		bool bPrintComma = false;
+		foreach_iter(exec_list_iterator, iter, *call)
+		{
+			ir_instruction *const inst = (ir_instruction *) iter.get();
+			if (bPrintComma)
+			{
+				ralloc_asprintf_append(buffer, ",");
+			}
+			inst->accept(this);
+			bPrintComma = true;
+		}
+		ralloc_asprintf_append(buffer, ")");
+	}
+
+	virtual void visit(ir_return *ret) override
+	{
+		check(scope_depth > 0);
+
+		ralloc_asprintf_append(buffer, "return ");
+		ir_rvalue *const value = ret->get_value();
+		if (value)
+		{
+			value->accept(this);
+		}
+	}
+
+	virtual void visit(ir_discard *discard) override
+	{
+		check(scope_depth > 0);
+
+		if (discard->condition)
+		{
+			ralloc_asprintf_append(buffer, "if (");
+			discard->condition->accept(this);
+			ralloc_asprintf_append(buffer, ") ");
+		}
+		ralloc_asprintf_append(buffer, "discard_fragment()");
+	}
+
+	bool try_conditional_move(ir_if *expr)
+	{
+		ir_dereference_variable *dest_deref = NULL;
+		ir_rvalue *true_value = NULL;
+		ir_rvalue *false_value = NULL;
+		unsigned write_mask = 0;
+		const glsl_type *assign_type = NULL;
+		int num_inst;
+
+		num_inst = 0;
+		foreach_iter(exec_list_iterator, iter, expr->then_instructions)
+		{
+			if (num_inst > 0)
+			{
+				// multiple instructions? not a conditional move
+				return false;
+			}
+
+			ir_instruction *const inst = (ir_instruction *) iter.get();
+			ir_assignment *assignment = inst->as_assignment();
+			if (assignment && (assignment->rhs->ir_type == ir_type_dereference_variable || assignment->rhs->ir_type == ir_type_constant))
+			{
+				dest_deref = assignment->lhs->as_dereference_variable();
+				true_value = assignment->rhs;
+				write_mask = assignment->write_mask;
+			}
+			num_inst++;
+		}
+
+		if (dest_deref == NULL || true_value == NULL)
+			return false;
+
+		num_inst = 0;
+		foreach_iter(exec_list_iterator, iter, expr->else_instructions)
+		{
+			if (num_inst > 0)
+			{
+				// multiple instructions? not a conditional move
+				return false;
+			}
+
+			ir_instruction *const inst = (ir_instruction *) iter.get();
+			ir_assignment *assignment = inst->as_assignment();
+			if (assignment && (assignment->rhs->ir_type == ir_type_dereference_variable || assignment->rhs->ir_type == ir_type_constant))
+			{
+				ir_dereference_variable *tmp_deref = assignment->lhs->as_dereference_variable();
+				if (tmp_deref
+					&& tmp_deref->var == dest_deref->var
+					&& tmp_deref->type == dest_deref->type
+					&& assignment->write_mask == write_mask)
+				{
+					false_value= assignment->rhs;
+				}
+			}
+			num_inst++;
+		}
+
+		if (false_value == NULL)
+			return false;
+
+		char mask[6];
+		unsigned j = 1;
+		if (dest_deref->type->is_scalar() == false || write_mask != 0x1)
+		{
+			for (unsigned i = 0; i < 4; i++)
+			{
+				if ((write_mask & (1 << i)) != 0)
+				{
+					mask[j] = "xyzw"[i];
+					j++;
+				}
+			}
+		}
+		mask[j] = '\0';
+		mask[0] = (j == 1) ? '\0' : '.';
+
+		dest_deref->accept(this);
+		ralloc_asprintf_append(buffer, "%s = (", mask);
+		expr->condition->accept(this);
+		ralloc_asprintf_append(buffer, ")?(");
+		true_value->accept(this);
+		ralloc_asprintf_append(buffer, "):(");
+		false_value->accept(this);
+		ralloc_asprintf_append(buffer, ")");
+
+		return true;
+	}
+
+	virtual void visit(ir_if *expr) override
+	{
+		check(scope_depth > 0);
+
+		if (try_conditional_move(expr) == false)
+		{
+		ralloc_asprintf_append(buffer, "if (");
+		expr->condition->accept(this);
+		ralloc_asprintf_append(buffer, ")\n");
+		indent();
+		ralloc_asprintf_append(buffer, "{\n");
+		
+		indentation++;
+		foreach_iter(exec_list_iterator, iter, expr->then_instructions)
+		{
+			ir_instruction *const inst = (ir_instruction *) iter.get();
+			indent();
+			do_visit(inst);
+		}
+		indentation--;
+
+		indent();
+		ralloc_asprintf_append(buffer, "}\n");
+
+		if (!expr->else_instructions.is_empty())
+		{
+			indent();
+			ralloc_asprintf_append(buffer, "else\n");
+			indent();
+			ralloc_asprintf_append(buffer, "{\n");
+
+			indentation++;
+			foreach_iter(exec_list_iterator, iter, expr->else_instructions)
+			{
+				ir_instruction *const inst = (ir_instruction *) iter.get();
+				indent();
+				do_visit(inst);
+			}
+			indentation--;
+
+			indent();
+			ralloc_asprintf_append(buffer, "}\n");
+		}
+
+		needs_semicolon = false;
+	}
+	}
+
+	virtual void visit(ir_loop *loop) override
+	{
+		check(scope_depth > 0);
+
+		if (loop->counter && loop->to)
+		{
+			// IR cmp operator is when to terminate loop; whereas GLSL for loop syntax
+			// is while to continue the loop. Invert the meaning of operator when outputting.
+			const char* termOp = NULL;
+			switch (loop->cmp)
+			{
+				case ir_binop_less: termOp = ">="; break;
+				case ir_binop_greater: termOp = "<="; break;
+				case ir_binop_lequal: termOp = ">"; break;
+				case ir_binop_gequal: termOp = "<"; break;
+				case ir_binop_equal: termOp = "!="; break;
+				case ir_binop_nequal: termOp = "=="; break;
+				default: check(false);
+			}
+			ralloc_asprintf_append(buffer, "for (;%s%s", unique_name(loop->counter), termOp);
+			loop->to->accept (this);
+			ralloc_asprintf_append(buffer, ";)\n");
+		}
+		else
+		{
+#if 1
+			ralloc_asprintf_append(buffer, "for (;;)\n");
+#else
+			ralloc_asprintf_append(buffer, "for ( int loop%d = 0; loop%d < 256; loop%d ++)\n", loop_count, loop_count, loop_count);
+			loop_count++;
+#endif
+		}
+		indent();
+		ralloc_asprintf_append(buffer, "{\n");
+
+		indentation++;
+		foreach_iter(exec_list_iterator, iter, loop->body_instructions)
+		{
+			ir_instruction *const inst = (ir_instruction *) iter.get();
+			indent();
+			do_visit(inst);
+		}
+		indentation--;
+
+		indent();
+		ralloc_asprintf_append(buffer, "}\n");
+
+		needs_semicolon = false;
+	}
+
+	virtual void visit(ir_loop_jump *jmp) override
+	{
+		check(scope_depth > 0);
+
+		ralloc_asprintf_append(buffer, "%s",
+			jmp->is_break() ? "break" : "continue");
+	}
+
+	virtual void visit(ir_atomic *ir) override
+	{
+		const char *sharedAtomicFunctions[] = 
+		{
+			"atomicAdd",
+			"atomicAnd",
+			"atomicMin",
+			"atomicMax",
+			"atomicOr",
+			"atomicXor",
+			"atomicExchange",
+			"atomicCompSwap"
+		};
+		const char *imageAtomicFunctions[] = 
+		{
+			"imageAtomicAdd",
+			"imageAtomicAnd",
+			"imageAtomicMin",
+			"imageAtomicMax",
+			"imageAtomicOr",
+			"imageAtomicXor",
+			"imageAtomicExchange",
+			"imageAtomicCompSwap"
+		};
+		check(scope_depth > 0);
+		const bool is_image = ir->memory_ref->as_dereference_image() != NULL;
+
+		ir->lhs->accept(this);
+		if (!is_image)
+		{
+			ralloc_asprintf_append(buffer, " = %s(",
+				sharedAtomicFunctions[ir->operation]);
+			ir->memory_ref->accept(this);
+			ralloc_asprintf_append(buffer, ", ");
+			ir->operands[0]->accept(this);
+			if (ir->operands[1])
+			{
+				ralloc_asprintf_append(buffer, ", ");
+				ir->operands[1]->accept(this);
+			}
+			ralloc_asprintf_append(buffer, ")");
+		}
+		else
+		{
+			ir_dereference_image *image = ir->memory_ref->as_dereference_image();
+			ralloc_asprintf_append(buffer, " = %s(",
+				imageAtomicFunctions[ir->operation]);
+			image->image->accept(this);
+			ralloc_asprintf_append(buffer, ", ");
+			image->image_index->accept(this);
+			ralloc_asprintf_append(buffer, ", ");
+			ir->operands[0]->accept(this);
+			if (ir->operands[1])
+			{
+				ralloc_asprintf_append(buffer, ", ");
+				ir->operands[1]->accept(this);
+			}
+			ralloc_asprintf_append(buffer, ")");
+		}
+	}
+
+	/**
+	 * Declare structs used to simulate multi-dimensional arrays.
+	 */
+	void declare_md_array_struct(const glsl_type* type, hash_table* ht)
+	{
+		check(type->is_array());
+
+		if (hash_table_find(ht, (void*)type) == NULL)
+		{
+			const glsl_type* subtype = type->fields.array;
+			if (subtype->base_type == GLSL_TYPE_ARRAY)
+			{
+				declare_md_array_struct(subtype, ht);
+
+				ralloc_asprintf_append(buffer, "struct ");
+				print_md_array_type(type);
+				ralloc_asprintf_append(buffer, "\n{\n\t");
+				print_md_array_type(subtype);
+				ralloc_asprintf_append(buffer, " Inner[%u];\n};\n\n", type->length);
+			}
+			else
+			{
+				ralloc_asprintf_append(buffer, "struct ");
+				print_md_array_type(type);
+				ralloc_asprintf_append(buffer, "\n{\n\t");
+				print_type_pre(type);
+				ralloc_asprintf_append(buffer, " Inner");
+				print_type_post(type);
+				ralloc_asprintf_append(buffer, ";\n};\n\n");
+			}
+			hash_table_insert(ht, (void*)type, (void*)type);
+		}
+	}
+
+	/**
+	 * Declare structs used by the code that has been generated.
+	 */
+	void declare_structs(_mesa_glsl_parse_state* state)
+	{
+		// If any variable in a uniform block is in use, the entire uniform block
+		// must be present including structs that are not actually accessed.
+		for (unsigned i = 0; i < state->num_uniform_blocks; i++)
+		{
+			const glsl_uniform_block* block = state->uniform_blocks[i];
+			if (hash_table_find(used_uniform_blocks, block->name))
+			{
+				for (unsigned var_index = 0; var_index < block->num_vars; ++var_index)
+				{
+					const glsl_type* type = block->vars[var_index]->type;
+					if (type->base_type == GLSL_TYPE_STRUCT &&
+						hash_table_find(used_structures, type) == NULL)
+					{
+						hash_table_insert(used_structures, (void*)type, type);
+					}
+				}
+			}
+		}
+
+		// If otherwise unused structure is a member of another, used structure, the unused structure is also, in fact, used
+		{
+			int added_structure_types;
+			do
+			{
+				added_structure_types = 0;
+				for (unsigned i = 0; i < state->num_user_structures; i++)
+				{
+					const glsl_type *const s = state->user_structures[i];
+
+					if (hash_table_find(used_structures, s) == NULL)
+					{
+						continue;
+					}
+
+					for (unsigned j = 0; j < s->length; j++)
+					{
+						const glsl_type* type = s->fields.structure[j].type;
+
+						if (type->base_type == GLSL_TYPE_STRUCT)
+						{
+							if (hash_table_find(used_structures, type) == NULL)
+							{
+								hash_table_insert(used_structures, (void*)type, type);
+								++added_structure_types;
+							}
+						}
+						else if (type->base_type == GLSL_TYPE_ARRAY && type->fields.array->base_type == GLSL_TYPE_STRUCT)
+						{
+							if (hash_table_find(used_structures, type->fields.array) == NULL)
+							{
+								hash_table_insert(used_structures, (void*)type->fields.array, type->fields.array);
+							}
+						}
+						else if ((type->base_type == GLSL_TYPE_INPUTPATCH || type->base_type == GLSL_TYPE_OUTPUTPATCH) && type->inner_type->base_type == GLSL_TYPE_STRUCT)
+						{
+							if (hash_table_find(used_structures, type->inner_type) == NULL)
+							{
+								hash_table_insert(used_structures, (void*)type->inner_type, type->inner_type);
+							}
+						}
+					}
+				}
+			}
+			while( added_structure_types > 0 );
+		}
+
+		// Generate structures that allow support for multi-dimensional arrays.
+		{
+			hash_table* ht = hash_table_ctor(32, hash_table_pointer_hash, hash_table_pointer_compare);
+			foreach_iter(exec_list_iterator, iter, used_md_arrays) 
+			{
+				md_array_entry* entry = (md_array_entry*)iter.get();
+				declare_md_array_struct(entry->type, ht);
+			}
+			hash_table_dtor(ht);
+		}
+
+		for (unsigned i = 0; i < state->num_user_structures; i++)
+		{
+			const glsl_type *const s = state->user_structures[i];
+
+			if (hash_table_find(used_structures, s) == NULL)
+			{
+				continue;
+			}
+
+			if (s->HlslName && !strcmp(s->HlslName, "__PACKED__"))
+			{
+				bUsePacked = true;
+			}
+
+			ralloc_asprintf_append(buffer, "struct %s\n{\n", s->name);
+
+			if (s->length == 0)
+			{
+				ralloc_asprintf_append(buffer, "\t float glsl_doesnt_like_empty_structs;\n");
+			}
+			else
+			{
+				for (unsigned j = 0; j < s->length; j++)
+				{
+					ralloc_asprintf_append(buffer, "\t");
+					print_type_pre(s->fields.structure[j].type);
+					ralloc_asprintf_append(buffer, " %s", s->fields.structure[j].name);
+					print_type_post(s->fields.structure[j].type);
+//@todo-rco
+					if (s->fields.structure[j].semantic)
+					{
+						if (!strncmp(s->fields.structure[j].semantic, "ATTRIBUTE", 9))
+						{
+							ralloc_asprintf_append(buffer, " [[ attribute(%s) ]]", s->fields.structure[j].semantic + 9);
+						}
+						else if (!strncmp(s->fields.structure[j].semantic, "gl_FragDepth", 12))
+						{
+							ralloc_asprintf_append(buffer, " [[ depth(any) ]]");
+						}
+						else if (!strncmp(s->fields.structure[j].semantic, "[[", 2))
+						{
+							ralloc_asprintf_append(buffer, " %s", s->fields.structure[j].semantic);
+						}
+						else
+						{
+							ralloc_asprintf_append(buffer, "[[ ERROR! ]]");
+							check(0);
+						}
+					}
+					ralloc_asprintf_append(buffer, ";\n");
+				}
+			}
+			ralloc_asprintf_append(buffer, "};\n\n");
+			bUsePacked = false;
+		}
+
+		unsigned num_used_blocks = 0;
+		for (unsigned i = 0; i < state->num_uniform_blocks; i++)
+		{
+			const glsl_uniform_block* block = state->uniform_blocks[i];
+			if (hash_table_find(used_uniform_blocks, block->name))
+			{
+				const char* block_name = block->name;
+				if (state->has_packed_uniforms)
+				{
+					block_name = ralloc_asprintf(mem_ctx, "%sb%u",
+						glsl_variable_tag_from_parser_target(state->target),
+						num_used_blocks
+						);
+				}
+				ralloc_asprintf_append(buffer, "layout(std140) uniform %s\n{\n", block_name);
+
+				for (unsigned var_index = 0; var_index < block->num_vars; ++var_index)
+				{
+					ir_variable* var = block->vars[var_index];
+					ralloc_asprintf_append(buffer, "\t");
+					print_type_pre(var->type);
+					ralloc_asprintf_append(buffer, " %s", var->name);
+					print_type_post(var->type);
+					ralloc_asprintf_append(buffer, ";\n");
+				}
+				ralloc_asprintf_append(buffer, "};\n\n");
+
+				num_used_blocks++;
+			}
+		}
+	}
+
+	void PrintPackedSamplers(_mesa_glsl_parse_state::TUniformList& Samplers, TStringToSetMap& TextureToSamplerMap)
+	{
+		bool bPrintHeader = true;
+		bool bNeedsComma = false;
+		for (_mesa_glsl_parse_state::TUniformList::iterator Iter = Samplers.begin(); Iter != Samplers.end(); ++Iter)
+		{
+			glsl_packed_uniform& Sampler = *Iter;
+			std::string SamplerStates("");
+			TStringToSetMap::iterator IterFound = TextureToSamplerMap.find(Sampler.Name);
+			if (IterFound != TextureToSamplerMap.end())
+			{
+				TStringSet& ListSamplerStates = IterFound->second;
+				check(!ListSamplerStates.empty());
+				for (TStringSet::iterator IterSS = ListSamplerStates.begin(); IterSS != ListSamplerStates.end(); ++IterSS)
+				{
+					if (IterSS == ListSamplerStates.begin())
+					{
+						SamplerStates += "[";
+					}
+					else
+					{
+						SamplerStates += ",";
+					}
+					SamplerStates += *IterSS;
+				}
+
+				SamplerStates += "]";
+			}
+
+			ralloc_asprintf_append(
+				buffer,
+				"%s%s(%u:%u%s)",
+				bNeedsComma ? "," : "",
+				Sampler.Name.c_str(),
+				Sampler.offset,
+				Sampler.num_components,
+				SamplerStates.c_str()
+				);
+
+			bNeedsComma = true;
+		}
+	}
+
+	bool PrintPackedUniforms(bool bPrintArrayType, char ArrayType, _mesa_glsl_parse_state::TUniformList& Uniforms, bool bFlattenUniformBuffers, bool NeedsComma)
+		{
+		bool bPrintHeader = true;
+		for (_mesa_glsl_parse_state::TUniformList::iterator Iter = Uniforms.begin(); Iter != Uniforms.end(); ++Iter)
+		{
+			glsl_packed_uniform& Uniform = *Iter;
+			if (!bFlattenUniformBuffers || Uniform.CB_PackedSampler.empty())
+			{
+				if (bPrintArrayType && bPrintHeader)
+				{
+					ralloc_asprintf_append(buffer, "%s%c[",
+						NeedsComma ? "," : "",
+						ArrayType);
+					bPrintHeader = false;
+					NeedsComma = false;
+				}
+				ralloc_asprintf_append(
+					buffer,
+					"%s%s(%u:%u)",
+					NeedsComma ? "," : "",
+					Uniform.Name.c_str(),
+					Uniform.offset,
+					Uniform.num_components
+					);
+				NeedsComma = true;
+				}
+			}
+
+		if (bPrintArrayType && !bPrintHeader)
+			{
+			ralloc_asprintf_append(buffer, "]");
+			}
+
+		return NeedsComma;
+		}
+
+	void PrintPackedGlobals(_mesa_glsl_parse_state* State)
+	{
+		//	@PackedGlobals: Global0(DestArrayType, DestOffset, SizeInFloats), Global1(DestArrayType, DestOffset, SizeInFloats), ...
+		bool bNeedsHeader = true;
+		bool bNeedsComma = false;
+		for (auto ArrayIter = State->GlobalPackedArraysMap.begin(); ArrayIter != State->GlobalPackedArraysMap.end(); ++ArrayIter)
+		{
+			char ArrayType = ArrayIter->first;
+			if (ArrayType != EArrayType_Image && ArrayType != EArrayType_Sampler)
+			{
+				_mesa_glsl_parse_state::TUniformList& Uniforms = ArrayIter->second;
+				check(!Uniforms.empty());
+
+				for (auto Iter = Uniforms.begin(); Iter != Uniforms.end(); ++Iter)
+				{
+					glsl_packed_uniform& Uniform = *Iter;
+					if (!State->bFlattenUniformBuffers || Uniform.CB_PackedSampler.empty())
+					{
+						if (bNeedsHeader)
+						{
+							ralloc_asprintf_append(buffer, "// @PackedGlobals: ");
+							bNeedsHeader = false;
+						}
+
+						ralloc_asprintf_append(
+							buffer,
+							"%s%s(%c:%u,%u)",
+							bNeedsComma ? "," : "",
+							Uniform.Name.c_str(),
+							ArrayType,
+							Uniform.offset,
+							Uniform.num_components
+							);
+						bNeedsComma = true;
+					}
+				}
+			}
+		}
+
+		if (!bNeedsHeader)
+		{
+			ralloc_asprintf_append(buffer, "\n");
+		}
+	}
+
+	void PrintPackedUniformBuffers(_mesa_glsl_parse_state* State)
+	{
+		// @PackedUB: UniformBuffer0(SourceIndex0): Member0(SourceOffset,SizeInFloats),Member1(SourceOffset,SizeInFloats), ...
+		// @PackedUB: UniformBuffer1(SourceIndex1): Member0(SourceOffset,SizeInFloats),Member1(SourceOffset,SizeInFloats), ...
+		// ...
+
+		// First find all used CBs (since we lost that info during flattening)
+		TStringSet UsedCBs;
+		for (auto IterCB = State->CBPackedArraysMap.begin(); IterCB != State->CBPackedArraysMap.end(); ++IterCB)
+		{
+			for (auto Iter = IterCB->second.begin(); Iter != IterCB->second.end(); ++Iter)
+			{
+				_mesa_glsl_parse_state::TUniformList& Uniforms = Iter->second;
+				for (auto IterU = Uniforms.begin(); IterU != Uniforms.end(); ++IterU)
+				{
+					if (!IterU->CB_PackedSampler.empty())
+					{
+						check(IterCB->first == IterU->CB_PackedSampler);
+						UsedCBs.insert(IterU->CB_PackedSampler);
+					}
+				}
+			}
+		}
+
+		check(UsedCBs.size() == State->CBPackedArraysMap.size());
+
+		// Now get the CB index based off source declaration order, and print an info line for each, while creating the mem copy list
+		unsigned CBIndex = 0;
+		TCBDMARangeMap CBRanges;
+		for (unsigned i = 0; i < State->num_uniform_blocks; i++)
+		{
+			const glsl_uniform_block* block = State->uniform_blocks[i];
+			if (UsedCBs.find(block->name) != UsedCBs.end())
+			{
+				bool bNeedsHeader = true;
+
+				// Now the members for this CB
+				bool bNeedsComma = false;
+				auto IterPackedArrays = State->CBPackedArraysMap.find(block->name);
+				check(IterPackedArrays != State->CBPackedArraysMap.end());
+				for (auto Iter = IterPackedArrays->second.begin(); Iter != IterPackedArrays->second.end(); ++Iter)
+				{
+					char ArrayType = Iter->first;
+					check(ArrayType != EArrayType_Image && ArrayType != EArrayType_Sampler);
+
+					_mesa_glsl_parse_state::TUniformList& Uniforms = Iter->second;
+					for (auto IterU = Uniforms.begin(); IterU != Uniforms.end(); ++IterU)
+					{
+						glsl_packed_uniform& Uniform = *IterU;
+						if (Uniform.CB_PackedSampler == block->name)
+						{
+							if (bNeedsHeader)
+							{
+								ralloc_asprintf_append(buffer, "// @PackedUB: %s(%d): ",
+									block->name,
+									CBIndex);
+								bNeedsHeader = false;
+							}
+
+							ralloc_asprintf_append(buffer, "%s%s(%u,%u)",
+								bNeedsComma ? "," : "",
+								Uniform.Name.c_str(),
+								Uniform.OffsetIntoCBufferInFloats,
+								Uniform.SizeInFloats);
+
+							bNeedsComma = true;
+							unsigned SourceOffset = Uniform.OffsetIntoCBufferInFloats;
+							unsigned DestOffset = Uniform.offset;
+							unsigned Size = Uniform.SizeInFloats;
+							unsigned DestCBIndex = 0;
+							unsigned DestCBPrecision = ArrayType;
+							InsertRange(CBRanges, CBIndex, SourceOffset, Size, DestCBIndex, DestCBPrecision, DestOffset);
+						}
+					}
+				}
+
+				if (!bNeedsHeader)
+				{
+					ralloc_asprintf_append(buffer, "\n");
+				}
+
+				CBIndex++;
+			}
+		}
+
+		//DumpSortedRanges(SortRanges(CBRanges));
+
+		// @PackedUBCopies: SourceArray:SourceOffset-DestArray:DestOffset,SizeInFloats;SourceArray:SourceOffset-DestArray:DestOffset,SizeInFloats,...
+		bool bFirst = true;
+		for (auto Iter = CBRanges.begin(); Iter != CBRanges.end(); ++Iter)
+		{
+			TDMARangeList& List = Iter->second;
+			for (auto IterList = List.begin(); IterList != List.end(); ++IterList)
+			{
+				if (bFirst)
+				{
+					ralloc_asprintf_append(buffer, "// @PackedUBGlobalCopies: ");
+					bFirst = false;
+				}
+				else
+				{
+					ralloc_asprintf_append(buffer, ",");
+				}
+
+				check(IterList->DestCBIndex == 0);
+				ralloc_asprintf_append(buffer, "%d:%d-%c:%d:%d", IterList->SourceCB, IterList->SourceOffset, IterList->DestCBPrecision, IterList->DestOffset, IterList->Size);
+			}
+		}
+
+		if (!bFirst)
+		{
+			ralloc_asprintf_append(buffer, "\n");
+		}
+	}
+
+	void PrintPackedUniforms(_mesa_glsl_parse_state* State)
+	{
+		PrintPackedGlobals(State);
+
+		if (State->bFlattenUniformBuffers && !State->CBuffersOriginal.empty())
+		{
+			PrintPackedUniformBuffers(State);
+		}
+	}
+
+	/**
+	 * Print a list of external variables.
+	 */
+	void print_extern_vars(_mesa_glsl_parse_state* State, exec_list* extern_vars)
+	{
+		const char *type_str[] = { "u", "i", "f", "f", "b", "t", "?", "?", "?", "?", "s", "os", "im", "ip", "op" };
+		const char *col_str[] = { "", "", "2x", "3x", "4x" };
+		const char *row_str[] = { "", "1", "2", "3", "4" };
+
+		check( sizeof(type_str)/sizeof(char*) == GLSL_TYPE_MAX);
+
+		bool need_comma = false;
+		foreach_iter(exec_list_iterator, iter, *extern_vars)
+		{
+			ir_variable* var = ((extern_var*)iter.get())->var;
+			const glsl_type* type = var->type;
+			if (!strcmp(var->name,"gl_in"))
+			{
+				// Ignore it, as we can't properly frame this information in current format, and it's not used anyway for geometry shaders
+				continue;
+			}
+			if (!strncmp(var->name,"in_",3) || !strncmp(var->name,"out_",4))
+			{
+				if (type->is_record())
+				{
+					// This is the specific case for GLSL >= 150, as we generate a struct with a member for each interpolator (which we still want to count)
+					if (type->length != 1)
+					{
+						_mesa_glsl_warning(State, "Found a complex structure as in/out, counting is not implemented yet...\n");
+						continue;
+					}
+
+					type = type->fields.structure->type;
+				}
+
+				// In and out variables may be packed in structures, or array of structures.
+				// But we're interested only in those that aren't, ie. inputs for vertex shader and outputs for pixel shader.
+				if (type->is_array() || type->is_record())
+				{
+					continue;
+				}
+			}
+			bool is_array = type->is_array();
+			int array_size = is_array ? type->length : 0;
+			if (is_array)
+			{
+				type = type->fields.array;
+			}
+			ralloc_asprintf_append(buffer, "%s%s%s%s",
+				need_comma ? "," : "",
+				type_str[type->base_type],
+				col_str[type->matrix_columns],
+				row_str[type->vector_elements]);
+			if (is_array)
+			{
+				ralloc_asprintf_append(buffer, "[%u]", array_size);
+			}
+			ralloc_asprintf_append(buffer, ":%s", var->name);
+			need_comma = true;
+		}
+	}
+
+	/**
+	 * Print the input/output signature for this shader.
+	 */
+	void print_signature(_mesa_glsl_parse_state *state)
+	{
+		if (!input_variables.is_empty())
+		{
+			ralloc_asprintf_append(buffer, "// @Inputs: ");
+			print_extern_vars(state, &input_variables);
+			ralloc_asprintf_append(buffer, "\n");
+		}
+
+		if (!output_variables.is_empty())
+		{
+			ralloc_asprintf_append(buffer, "// @Outputs: ");
+			print_extern_vars(state, &output_variables);
+			ralloc_asprintf_append(buffer, "\n");
+		}
+		if (state->num_uniform_blocks > 0 && !state->bFlattenUniformBuffers)
+		{
+			bool bFirst = true;
+			for (int i = 0; i < Buffers.Buffers.Num(); ++i)
+			{
+				// Some entries might be null, if we used more packed than real UBs used
+				if (Buffers.Buffers[i])
+				{
+					auto* Var = Buffers.Buffers[i]->as_variable();
+					if (!Var->semantic && !Var->type->is_sampler())
+					{
+						ralloc_asprintf_append(buffer, "%s%s(%d)",
+							bFirst ? "// @UniformBlocks: " : ",",
+							Var->name, i);
+						bFirst = false;
+					}
+				}
+			}
+			if (!bFirst)
+			{
+				ralloc_asprintf_append(buffer, "\n");
+			}
+		}
+
+		if (state->has_packed_uniforms)
+		{
+			PrintPackedUniforms(state);
+
+			if (!state->GlobalPackedArraysMap[EArrayType_Sampler].empty())
+			{
+				ralloc_asprintf_append(buffer, "// @Samplers: ");
+				PrintPackedSamplers(
+					state->GlobalPackedArraysMap[EArrayType_Sampler],
+					state->TextureToSamplerMap
+					);
+				ralloc_asprintf_append(buffer, "\n");
+			}
+
+			if (!state->GlobalPackedArraysMap[EArrayType_Image].empty())
+			{
+				ralloc_asprintf_append(buffer, "// @UAVs: ");
+				PrintPackedUniforms(
+					false,
+					EArrayType_Image,
+					state->GlobalPackedArraysMap[EArrayType_Image],
+					false,
+					false
+					);
+				ralloc_asprintf_append(buffer, "\n");
+			}
+		}
+		else
+		{
+			if (!uniform_variables.is_empty())
+			{
+				ralloc_asprintf_append(buffer, "// @Uniforms: ");
+				print_extern_vars(state, &uniform_variables);
+				ralloc_asprintf_append(buffer, "\n");
+			}
+			if (!sampler_variables.is_empty())
+			{
+				ralloc_asprintf_append(buffer, "// @Samplers: ");
+				print_extern_vars(state, &sampler_variables);
+				ralloc_asprintf_append(buffer, "\n");
+			}
+			if (!image_variables.is_empty())
+			{
+				ralloc_asprintf_append(buffer, "// @UAVs: ");
+				print_extern_vars(state, &image_variables);
+				ralloc_asprintf_append(buffer, "\n");
+			}
+		}
+	}
+
+	/**
+	 * Print the layout directives for this shader.
+	 */
+	void print_layout(_mesa_glsl_parse_state *state)
+	{
+		if (state->target == compute_shader )
+		{
+			ralloc_asprintf_append(buffer, "layout( local_size_x = %d, "
+				"local_size_y = %d, local_size_z = %d ) in;\n", wg_size_x,
+				wg_size_y, wg_size_z );
+		}
+
+		if(state->target == tessellation_evaluation_shader || state->target == tessellation_control_shader)
+		{
+			check(0);
+		}
+	}
+
+public:
+
+	/** Constructor. */
+	FGenerateMetalVisitor(_mesa_glsl_parse_state* InParseState, bool bInIsVS, FBuffers& InBuffers)
+		: ParseState(InParseState)
+		, buffer(0)
+		, bIsVS(bInIsVS)
+		, Buffers(InBuffers)
+		, indentation(0)
+		, temp_id(0)
+		, global_id(0)
+		, scope_depth(0)
+		, needs_semicolon(false)
+		, IsMain(false)
+		, should_print_uint_literals_as_ints(false)
+		, loop_count(0)
+		, bStageInEmitted(false)
+		, bUsePacked(false)
+	{
+		printable_names = hash_table_ctor(32, hash_table_pointer_hash, hash_table_pointer_compare);
+		used_structures = hash_table_ctor(128, hash_table_pointer_hash, hash_table_pointer_compare);
+		used_uniform_blocks = hash_table_ctor(32, hash_table_string_hash, hash_table_string_compare);
+	}
+
+	/** Destructor. */
+	virtual ~FGenerateMetalVisitor()
+	{
+		hash_table_dtor(printable_names);
+		hash_table_dtor(used_structures);
+		hash_table_dtor(used_uniform_blocks);
+	}
+
+	/**
+	 * Executes the visitor on the provided ir.
+	 * @returns the Metal source code generated.
+	 */
+	const char* run(exec_list* ir)
+	{
+		mem_ctx = ralloc_context(NULL);
+
+		char* code_buffer = ralloc_asprintf(mem_ctx, "");
+		buffer = &code_buffer;
+
+		foreach_iter(exec_list_iterator, iter, *ir)
+		{
+			ir_instruction *inst = (ir_instruction *)iter.get();
+			do_visit(inst);
+		}
+		buffer = 0;
+
+		char* decl_buffer = ralloc_asprintf(mem_ctx, "");
+		buffer = &decl_buffer;
+		declare_structs(ParseState);
+		buffer = 0;
+
+		char* signature = ralloc_asprintf(mem_ctx, "");
+		buffer = &signature;
+		print_signature(ParseState);
+		buffer = 0;
+
+		char* layout = ralloc_asprintf(mem_ctx, "");
+		buffer = &layout;
+		print_layout(ParseState);
+		buffer = 0;
+		char* full_buffer = ralloc_asprintf(
+			ParseState,
+			"// Compiled by HLSLCC\n%s\n#include <metal_stdlib>\nusing namespace metal;\n\n%s%s%s",
+			signature,
+			layout,
+			decl_buffer,
+			code_buffer
+			);
+		ralloc_free(mem_ctx);
+
+		return full_buffer;
+	}
+};
+
+char* FMetalCodeBackend::GenerateCode(exec_list* ir, _mesa_glsl_parse_state* state, EHlslShaderFrequency Frequency)
+{
+	// We'll need this Buffers info for the [[buffer()]] index
+	FBuffers Buffers;
+	FGenerateMetalVisitor visitor(state, (state->target == vertex_shader), Buffers);
+
+	// At this point, all inputs and outputs are global uniforms, no structures.
+
+	// Promotes all inputs from half to float to avoid stage_in issues
+	PromoteInputsAndOutputsGlobalHalfToFloat(ir, state, Frequency);
+
+	// Move all inputs & outputs to structs for Metal
+	PackInputsAndOutputs(ir, state, Frequency, visitor.input_variables);
+
+	// ir_var_uniform instances be global, so move them as arguments to main
+	MovePackedUniformsToMain(ir, state, Buffers);
+
+	//@todo-rco: Do we need this here?
+	ExpandArrayAssignments(ir, state);
+
+	// Fix any special language extensions (FrameBufferFetchES2() intrinsic)
+	FixIntrinsics(ir, state);
+
+	// Remove half->float->half or float->half->float
+	FixRedundantCasts(ir);
+
+	if (!OptimizeAndValidate(ir, state))
+	{
+		return nullptr;
+	}
+
+	// Do not call Optimize() after this!
+	{
+		// Metal can't do implicit conversions between half<->float during math expressions
+		BreakPrecisionChangesVisitor(ir, state);
+
+		// Metal can't read from a packed_* type, which for us come from a constant buffer
+		//@todo-rco: Might not work if accessing packed_half* m[N]!
+		RemovePackedVarReferences(ir, state);
+
+		bool bConvertUniformsToFloats = (HlslCompileFlags & HLSLCC_FlattenUniformBuffers) != HLSLCC_FlattenUniformBuffers;
+		ConvertHalfToFloatUniformsAndSamples(ir, state, bConvertUniformsToFloats, true);
+
+		Validate(ir, state);
+	}
+
+	// Generate the actual code string
+	const char* code = visitor.run(ir);
+	return _strdup(code);
+}
+
+struct FMetalUnsupportedSamplerVisitor : public ir_hierarchical_visitor
+{
+	const bool bIsVertexShader;
+	_mesa_glsl_parse_state* ParseState;
+	FMetalUnsupportedSamplerVisitor(_mesa_glsl_parse_state* InParseState, bool bInIsVertexShader) :
+		ParseState(InParseState),
+		bIsVertexShader(bInIsVertexShader)
+	{
+	}
+
+	virtual ir_visitor_status visit_leave(ir_texture* IR) override
+	{
+
+		return visit_continue;
+	}
+}; 
+
+
+// Converts an array index expression using an integer input attribute, to a float input attribute using a conversion to int
+struct SConvertIntVertexAttributeES2 : public ir_hierarchical_visitor
+{
+	_mesa_glsl_parse_state* ParseState;
+	exec_list* FunctionBody;
+	int InsideArrayDeref;
+	std::map<ir_variable*, ir_variable*> ConvertedVarMap;
+
+	SConvertIntVertexAttributeES2(_mesa_glsl_parse_state* InParseState, exec_list* InFunctionBody) : ParseState(InParseState), FunctionBody(InFunctionBody), InsideArrayDeref(0)
+	{
+	}
+
+	virtual ir_visitor_status visit_enter(ir_dereference_array* DeRefArray) override
+	{
+		// Break the array dereference so we know we want to modify the array index part
+		auto Result = ir_hierarchical_visitor::visit_enter(DeRefArray);
+		++InsideArrayDeref;
+		DeRefArray->array_index->accept(this);
+		--InsideArrayDeref;
+
+		return visit_continue;
+	}
+
+	virtual ir_visitor_status visit(ir_dereference_variable* DeRefVar) override
+	{
+		if (InsideArrayDeref > 0)
+		{
+			ir_variable* SourceVar = DeRefVar->var;
+			if (SourceVar->mode == ir_var_in)
+			{
+				// First time it still is an integer, so add the temporary and a conversion, and switch to float
+				if (SourceVar->type->is_integer())
+				{
+					check(SourceVar->type->is_integer() && !SourceVar->type->is_matrix() && !SourceVar->type->is_array());
+
+					// Double check we haven't processed this
+					auto IterFound = ConvertedVarMap.find(SourceVar);
+					check(IterFound == ConvertedVarMap.end());
+
+					// New temp var
+					ir_variable* NewVar = new(ParseState)ir_variable(SourceVar->type, NULL, ir_var_temporary);
+					base_ir->insert_before(NewVar);
+
+					// Switch original type to float
+					SourceVar->type = glsl_type::get_instance(GLSL_TYPE_FLOAT, SourceVar->type->vector_elements, 1);
+
+					// Convert float to int
+					ir_dereference_variable* NewSourceDeref = new(ParseState)ir_dereference_variable(SourceVar);
+					ir_expression* NewCastExpression = new(ParseState)ir_expression(ir_unop_f2i, NewSourceDeref);
+					ir_assignment* NewAssigment = new(ParseState)ir_assignment(new(ParseState)ir_dereference_variable(NewVar), NewCastExpression);
+					base_ir->insert_before(NewAssigment);
+
+					// Add the entry and modify the original Var
+					ConvertedVarMap[SourceVar] = NewVar;
+					DeRefVar->var = NewVar;
+				}
+				else
+				{
+					auto IterFound = ConvertedVarMap.find(SourceVar);
+					if (IterFound != ConvertedVarMap.end())
+					{
+						DeRefVar->var = IterFound->second;
+					}
+				}
+			}
+		}
+
+		return ir_hierarchical_visitor::visit(DeRefVar);
+	}
+};
+
+
+bool FMetalCodeBackend::ApplyAndVerifyPlatformRestrictions(exec_list* Instructions, _mesa_glsl_parse_state* ParseState, EHlslShaderFrequency Frequency)
+{
+	bool bIsVertexShader = (Frequency == HSF_VertexShader);
+
+	const bool bAllowVertexTextureFetch = true;
+	if (!bAllowVertexTextureFetch)
+	{
+		FMetalUnsupportedSamplerVisitor Visitor(ParseState, bIsVertexShader);
+		Visitor.run(Instructions);
+	}
+
+	// Handle integer vertex attributes used as array indices
+	if (bIsVertexShader)
+	{
+		SConvertIntVertexAttributeES2 ConvertIntVertexAttributeVisitor(ParseState, Instructions);
+		ConvertIntVertexAttributeVisitor.run(Instructions);
+	}
+
+	return true;
+}
+
+bool FMetalCodeBackend::GenerateMain(EHlslShaderFrequency Frequency, const char* EntryPoint, exec_list* Instructions, _mesa_glsl_parse_state* ParseState)
+{
+	auto* EntryPointSig = FindEntryPointFunction(Instructions, ParseState, EntryPoint);
+	if (!EntryPointSig)
+	{
+		_mesa_glsl_error(ParseState, "shader entry point '%s' not found", EntryPoint);
+		return false;
+	}
+
+	exec_list DeclInstructions;
+	exec_list PreCallInstructions;
+	exec_list ArgInstructions;
+	exec_list PostCallInstructions;
+
+	ParseState->symbols->push_scope();
+
+	// Find all system semantics and generate in/out globals
+	foreach_iter(exec_list_iterator, Iter, EntryPointSig->parameters)
+	{
+		const ir_variable* Variable = (ir_variable*)Iter.get();
+		if (Variable->semantic != NULL || Variable->type->is_record())
+		{
+			ir_dereference_variable* ArgVarDeref = NULL;
+			switch (Variable->mode)
+			{
+			case ir_var_in:
+				ArgVarDeref = MetalUtils::GenerateInput(
+					Frequency,
+					ParseState,
+					Variable->semantic,
+					Variable->type,
+					&DeclInstructions,
+					&PreCallInstructions
+					);
+				break;
+			case ir_var_out:
+				ArgVarDeref = MetalUtils::GenerateOutput(
+					Frequency,
+					ParseState,
+					Variable->semantic,
+					Variable->type,
+					&DeclInstructions,
+					&PreCallInstructions,
+					&PostCallInstructions
+					);
+				break;
+			default:
+			   _mesa_glsl_error(
+				   ParseState,
+				   "entry point parameter '%s' must be an input or output",
+				   Variable->name
+				   );
+			}
+			
+			ArgInstructions.push_tail(ArgVarDeref);
+		}
+	}
+
+
+	// The function's return value should have an output semantic if it's not void.
+	ir_dereference_variable* EntryPointReturn = nullptr;
+	if (!EntryPointSig->return_type->is_void())
+	{
+		EntryPointReturn = MetalUtils::GenerateOutput(Frequency, ParseState, EntryPointSig->return_semantic, EntryPointSig->return_type, &DeclInstructions, &PreCallInstructions, &PostCallInstructions);
+	}
+
+	ParseState->symbols->pop_scope();
+
+	// Generate the Main() function signature
+	ir_function_signature* MainSig = new(ParseState) ir_function_signature(glsl_type::void_type);
+	MainSig->is_defined = true;
+	MainSig->is_main = true;
+	MainSig->body.append_list(&PreCallInstructions);
+	// Call the original EntryPoint
+	MainSig->body.push_tail(new(ParseState) ir_call(EntryPointSig, EntryPointReturn, &ArgInstructions));
+	MainSig->body.append_list(&PostCallInstructions);
+
+	// Generate the Main() function
+	auto* MainFunction = new(ParseState)ir_function("Main");
+	MainFunction->add_signature(MainSig);
+	// Adds uniforms as globals
+	Instructions->append_list(&DeclInstructions);
+	Instructions->push_tail(MainFunction);
+
+	// Now that we have a proper Main(), move global setup to Main().
+	MoveGlobalInstructionsToMain(Instructions);
+//IRDump(Instructions);
+	return true;
+}
+
+FMetalCodeBackend::FMetalCodeBackend(unsigned int InHlslCompileFlags) :
+	FCodeBackend(InHlslCompileFlags)
+{
+}
+
+void FMetalLanguageSpec::SetupLanguageIntrinsics(_mesa_glsl_parse_state* State, exec_list* ir)
+{
+	// Framebuffer fetch
+	{
+		// Leave original fb ES2 fetch function as that's what the hlsl expects
+		make_intrinsic_genType(ir, State, FRAMEBUFFER_FETCH_ES2, ir_invalid_opcode, IR_INTRINSIC_HALF, 0, 4, 4);
+
+		// MRTs; first make intrinsics for each MRT, then a non-intrinsic version to use that (helps when converting to Metal)
+		for (int i = 0; i < MAX_SIMULTANEOUS_RENDER_TARGETS; ++i)
+		{
+			char FunctionName[32];
+			sprintf(FunctionName, "%s%d", FRAMEBUFFER_FETCH_MRT, i);
+			make_intrinsic_genType(ir, State, FunctionName, ir_invalid_opcode, IR_INTRINSIC_HALF, 0, 4, 4);
+		}
+
+		const auto* ReturnType = glsl_type::get_instance(GLSL_TYPE_HALF, 4, 1);
+		ir_function* Func = new(State)ir_function(FRAMEBUFFER_FETCH_MRT);
+		ir_function_signature* Sig = new(State) ir_function_signature(ReturnType);
+		//Sig->is_builtin = true;
+		Sig->is_defined = true;
+		ir_variable* MRTIndex = new(State) ir_variable(glsl_type::int_type, "Arg0", ir_var_in);
+		Sig->parameters.push_tail(MRTIndex);
+
+		for (int i = 0; i < MAX_SIMULTANEOUS_RENDER_TARGETS; ++i)
+		{
+			// Inject:
+			//	if (Arg0 == i) FRAMEBUFFER_FETCH_MRT#i();
+			auto* Condition = new(State) ir_expression(ir_binop_equal, new(State) ir_dereference_variable((ir_variable*)Sig->parameters.get_head()), new(State) ir_constant(i));
+			auto* If = new(State) ir_if(Condition);
+			char FunctionName[32];
+			sprintf(FunctionName, "%s%d", FRAMEBUFFER_FETCH_MRT, i);
+			auto* IntrinsicSig = FCodeBackend::FindEntryPointFunction(ir, State, FunctionName);
+			auto* ReturnValue = new(State) ir_variable(ReturnType, nullptr, ir_var_temporary);
+			exec_list Empty;
+			auto* Call = new(State) ir_call(IntrinsicSig, new(State) ir_dereference_variable(ReturnValue), &Empty);
+			Call->use_builtin = true;
+			If->then_instructions.push_tail(ReturnValue);
+			If->then_instructions.push_tail(Call);
+			If->then_instructions.push_tail(new(State) ir_return(new(State) ir_dereference_variable(ReturnValue)));
+			Sig->body.push_tail(If);
+		}
+
+		Func->add_signature(Sig);
+
+		State->symbols->add_global_function(Func);
+		ir->push_tail(Func);
+	}
+}
diff --git a/Engine/Source/Developer/iOS/MetalShaderFormat/Private/MetalBackend.h b/Engine/Source/Developer/iOS/MetalShaderFormat/Private/MetalBackend.h
new file mode 100644
--- /dev/null
+++ b/Engine/Source/Developer/iOS/MetalShaderFormat/Private/MetalBackend.h
@@ -0,0 +1,51 @@
+// Copyright 1998-2014 Epic Games, Inc. All Rights Reserved.
+
+#pragma once
+
+#include "hlslcc.h"
+#include "LanguageSpec.h"
+
+
+class FMetalLanguageSpec : public ILanguageSpec
+{
+public:
+	virtual bool SupportsDeterminantIntrinsic() const
+	{
+		//@todo-rco: Temp workaround for Seed 2 & 3
+		return false;// true;
+	}
+
+	virtual bool SupportsTransposeIntrinsic() const
+	{
+		//@todo-rco: Temp workaround for Seed 2 & 3
+		return false;// true;
+	}
+	virtual bool SupportsIntegerModulo() const { return true; }
+
+	virtual bool SupportsMatrixConversions() const { return false; }
+
+	virtual void SetupLanguageIntrinsics(_mesa_glsl_parse_state* State, exec_list* ir) override;
+};
+
+struct FBuffers;
+
+// Generates Metal compliant code from IR tokens
+struct FMetalCodeBackend : public FCodeBackend
+{
+	FMetalCodeBackend(unsigned int InHlslCompileFlags);
+
+	virtual char* GenerateCode(struct exec_list* ir, struct _mesa_glsl_parse_state* ParseState, EHlslShaderFrequency Frequency) override;
+
+	virtual bool GenerateMain(EHlslShaderFrequency Frequency, const char* EntryPoint, exec_list* Instructions, _mesa_glsl_parse_state* ParseState) override;
+
+	// Return false if there were restrictions that made compilation fail
+	virtual bool ApplyAndVerifyPlatformRestrictions(exec_list* Instructions, _mesa_glsl_parse_state* ParseState, EHlslShaderFrequency Frequency) override;
+
+	void PackInputsAndOutputs(exec_list* ir, _mesa_glsl_parse_state* state, EHlslShaderFrequency Frequency, exec_list& InputVars);
+	void MovePackedUniformsToMain(exec_list* ir, _mesa_glsl_parse_state* state, FBuffers& OutBuffers);
+	void FixIntrinsics(exec_list* ir, _mesa_glsl_parse_state* state);
+	void RemovePackedVarReferences(exec_list* ir, _mesa_glsl_parse_state* State);
+	void PromoteInputsAndOutputsGlobalHalfToFloat(exec_list* ir, _mesa_glsl_parse_state* state, EHlslShaderFrequency Frequency);
+	void ConvertHalfToFloatUniformsAndSamples(exec_list* ir, _mesa_glsl_parse_state* State, bool bConvertUniforms, bool bConvertSamples);
+	void BreakPrecisionChangesVisitor(exec_list* ir, _mesa_glsl_parse_state* State);
+};
diff --git a/Engine/Source/Developer/iOS/MetalShaderFormat/Private/MetalShaderCompiler.cpp b/Engine/Source/Developer/iOS/MetalShaderFormat/Private/MetalShaderCompiler.cpp
--- a/Engine/Source/Developer/iOS/MetalShaderFormat/Private/MetalShaderCompiler.cpp
+++ b/Engine/Source/Developer/iOS/MetalShaderFormat/Private/MetalShaderCompiler.cpp
@@ -19,7 +19,7 @@
 #endif
 #include "ShaderPreprocessor.h"
 #include "hlslcc.h"
-#include "Metal/MetalBackend.h"
+#include "MetalBackend.h"
 
 DEFINE_LOG_CATEGORY_STATIC(LogMetalShaderCompiler, Log, All); 
 
diff --git a/Engine/Source/Developer/iOS/MetalShaderFormat/Private/MetalUtils.cpp b/Engine/Source/Developer/iOS/MetalShaderFormat/Private/MetalUtils.cpp
new file mode 100644
--- /dev/null
+++ b/Engine/Source/Developer/iOS/MetalShaderFormat/Private/MetalUtils.cpp
@@ -0,0 +1,2376 @@
+// Copyright 1998-2014 Epic Games, Inc. All Rights Reserved.
+
+#include "MetalShaderFormat.h"
+#include "hlslcc.h"
+#include "hlslcc_private.h"
+#include "MetalBackend.h"
+#include "glsl/ir_gen_glsl.h"
+#include "mesa/compiler.h"
+#include "mesa/glsl_parser_extras.h"
+#include "mesa/hash_table.h"
+#include "mesa/ir_rvalue_visitor.h"
+#include "mesa/ast.h"
+#include "PackUniformBuffers.h"
+#include "IRDump.h"
+#include "OptValueNumbering.h"
+#include "mesa/ir_optimization.h"
+#include "MetalUtils.h"
+
+const bool bExpandVSInputsToFloat4 = false;
+const bool bGenerateVSInputDummies = false;
+
+static int GetIndexSuffix(const char* Prefix, int PrefixLength, const char* Semantic)
+{
+	check(Semantic);
+
+	if (!strncmp(Semantic, Prefix, PrefixLength))
+	{
+		Semantic += PrefixLength;
+		int Index = 0;
+		if (isdigit(*Semantic))
+		{
+			Index = (*Semantic) - '0';
+
+			++Semantic;
+			if (*Semantic == 0)
+			{
+				return Index;
+			}
+			else if (isdigit(*Semantic))
+			{
+				Index = Index * 10 + (*Semantic) - '0';
+				++Semantic;
+				if (*Semantic == 0)
+				{
+					return Index;
+				}
+			}
+		}
+	}
+
+	return -1;
+}
+
+static int GetAttributeIndex(const char* Semantic)
+{
+	return GetIndexSuffix("ATTRIBUTE", 9, Semantic);
+}
+
+static int GetInAttributeIndex(const char* Semantic)
+{
+	return GetIndexSuffix("IN_ATTRIBUTE", 12, Semantic);
+}
+
+const glsl_type* PromoteHalfToFloatType(_mesa_glsl_parse_state* state, const glsl_type* type)
+{
+	if (type->base_type == GLSL_TYPE_HALF)
+	{
+		return glsl_type::get_instance(GLSL_TYPE_FLOAT, type->vector_elements, type->matrix_columns);
+	}
+	else if (type->is_array())
+	{
+		auto* ElementType = type->element_type();
+		auto* NewElementType = PromoteHalfToFloatType(state, ElementType);
+		if (NewElementType != ElementType)
+		{
+			return glsl_type::get_array_instance(NewElementType, type->length);
+		}
+	}
+	else if (type->is_record())
+	{
+		auto* Fields = ralloc_array(state, glsl_struct_field, type->length);
+		bool bNeedNewType = false;
+		for (unsigned i = 0; i < type->length; ++i)
+		{
+			auto* NewMemberType = PromoteHalfToFloatType(state, type->fields.structure[i].type);
+			Fields[i] = type->fields.structure[i];
+			if (NewMemberType != type->fields.structure[i].type)
+			{
+				bNeedNewType = true;
+				Fields[i].type = NewMemberType;
+			}
+		}
+
+		if (bNeedNewType)
+		{
+			auto* NewType = glsl_type::get_record_instance(Fields, type->length, ralloc_asprintf(state, "%s_F", type->name));
+			// Hack: This way we tell this is a uniform buffer and we need to emit 'packed_'
+			((glsl_type*)NewType)->HlslName = "__PACKED__";
+			state->AddUserStruct(NewType);
+			return NewType;
+		}
+	}
+
+	return type;
+}
+
+void CreateNewAssignmentsFloat2Half(_mesa_glsl_parse_state* State, exec_list& NewAssignments, ir_variable* NewVar, ir_rvalue* RValue)
+{
+	if (NewVar->type->is_matrix())
+	{
+		for (uint32 i = 0; i < NewVar->type->matrix_columns; ++i)
+		{
+			auto* NewF2H = new(State)ir_expression(ir_unop_f2h, new(State)ir_dereference_array(RValue, new(State)ir_constant(i)));
+			auto* NewAssignment = new(State)ir_assignment(new(State)ir_dereference_array(NewVar, new(State)ir_constant(i)), NewF2H);
+			NewAssignments.push_tail(NewAssignment);
+		}
+	}
+	else
+	{
+		auto* NewF2H = new(State)ir_expression(ir_unop_f2h, RValue);
+		auto* NewAssignment = new(State)ir_assignment(new(State)ir_dereference_variable(NewVar), NewF2H);
+
+		NewAssignments.push_tail(NewAssignment);
+	}
+}
+
+static void CreateNewAssignmentsHalf2Float(_mesa_glsl_parse_state* State, exec_list& NewAssignments, ir_variable* NewVar, ir_rvalue* RValue)
+{
+	if (NewVar->type->is_matrix())
+	{
+		for (uint32 i = 0; i < NewVar->type->matrix_columns; ++i)
+		{
+			auto* NewF2H = new(State)ir_expression(ir_unop_h2f, new(State)ir_dereference_array(RValue, new(State)ir_constant(i)));
+			auto* NewAssignment = new(State)ir_assignment(new(State)ir_dereference_array(NewVar, new(State)ir_constant(i)), NewF2H);
+			NewAssignments.push_tail(NewAssignment);
+		}
+	}
+	else
+	{
+		auto* NewF2H = new(State)ir_expression(ir_unop_h2f, RValue);
+		auto* NewAssignment = new(State)ir_assignment(new(State)ir_dereference_variable(NewVar), NewF2H);
+
+		NewAssignments.push_tail(NewAssignment);
+	}
+}
+
+const glsl_type* GetFragColorTypeFromMetalOutputStruct(const glsl_type* OutputType)
+{
+	const glsl_type* FragColorType = glsl_type::error_type;
+	if (OutputType && OutputType->base_type == GLSL_TYPE_STRUCT)
+	{
+		for (unsigned j = 0; j < OutputType->length; j++)
+		{
+			if (OutputType->fields.structure[j].semantic)
+			{
+				//@todo-rco: MRTs?
+				if (!strncmp(OutputType->fields.structure[j].semantic, "[[ color(", 9))
+				{
+					FragColorType = OutputType->fields.structure[j].type;
+					break;
+				}
+			}
+		}
+	}
+	return FragColorType;
+}
+
+namespace MetalUtils
+{
+	/** Information on system values. */
+	struct FSystemValue
+	{
+		const char* HlslSemantic;
+		const glsl_type* Type;
+		const char* MetalName;
+		ir_variable_mode Mode;
+		const char* MetalSemantic;
+	};
+
+	/** Vertex shader system values. */
+	static FSystemValue VertexSystemValueTable[] =
+	{
+		{"SV_VertexID", glsl_type::uint_type, "IN_VertexID", ir_var_in, "[[ vertex_id ]]"},
+		{"SV_InstanceID", glsl_type::uint_type, "IN_InstanceID", ir_var_in, "[[ instance_id ]]"},
+		{"SV_Position", glsl_type::vec4_type, "Position", ir_var_out, "[[ position ]]"},
+		{NULL, NULL, NULL, ir_var_auto, nullptr}
+	};
+
+	/** Pixel shader system values. */
+	static FSystemValue PixelSystemValueTable[] =
+	{
+		{"SV_Depth", glsl_type::float_type, "FragDepth", ir_var_out, "[[ depth(any) ]]"},
+		{"SV_Position", glsl_type::vec4_type, "IN_FragCoord", ir_var_in, "[[ position ]]"},
+		{"SV_IsFrontFace", glsl_type::bool_type, "IN_FrontFacing", ir_var_in, "[[ front_facing ]]"},
+		//{"SV_PrimitiveID", glsl_type::int_type, "IN_PrimitiveID", ir_var_in, "[[  ]]"},
+		//{"SV_RenderTargetArrayIndex", glsl_type::int_type, "IN_Layer", ir_var_in, nullptr},
+		{"SV_Target0", glsl_type::half4_type, "FragColor0", ir_var_out, "[[ color(0) ]]"},
+		{"SV_Target1", glsl_type::half4_type, "FragColor1", ir_var_out, "[[ color(1) ]]"},
+		{"SV_Target2", glsl_type::half4_type, "FragColor2", ir_var_out, "[[ color(2) ]]"},
+		{"SV_Target3", glsl_type::half4_type, "FragColor3", ir_var_out, "[[ color(3) ]]"},
+		{"SV_Target4", glsl_type::half4_type, "FragColor4", ir_var_out, "[[ color(4) ]]"},
+		{"SV_Target5", glsl_type::half4_type, "FragColor5", ir_var_out, "[[ color(5) ]]"},
+		{"SV_Target6", glsl_type::half4_type, "FragColor6", ir_var_out, "[[ color(6) ]]"},
+		{"SV_Target7", glsl_type::half4_type, "FragColor7", ir_var_out, "[[ color(7) ]]"},
+		{NULL, NULL, NULL, ir_var_auto, nullptr}
+	};
+
+	/** Geometry shader system values. */
+	static FSystemValue GeometrySystemValueTable[] =
+	{
+/*
+		{"SV_VertexID", glsl_type::int_type, "IN_VertexID", ir_var_in, false, false, false, false},
+		{"SV_InstanceID", glsl_type::int_type, "IN_InstanceID", ir_var_in, false, false, false, false},
+		{"SV_Position", glsl_type::vec4_type, "IN_Position", ir_var_in, false, true, true, false},
+		{"SV_Position", glsl_type::vec4_type, "OUT_Position", ir_var_out, false, false, true, false},
+		{"SV_RenderTargetArrayIndex", glsl_type::int_type, "OUT_Layer", ir_var_out, false, false, false, false},
+		{"SV_PrimitiveID", glsl_type::int_type, "OUT_PrimitiveID", ir_var_out, false, false, false, false},
+		{"SV_PrimitiveID", glsl_type::int_type, "IN_PrimitiveIDIn", ir_var_in, false, false, false, false},
+*/
+		{NULL, NULL, NULL, ir_var_auto, nullptr}
+	};
+
+
+	/** Hull shader system values. */
+	static FSystemValue HullSystemValueTable[] =
+	{
+/*
+		{"SV_OutputControlPointID", glsl_type::int_type, "gl_InvocationID", ir_var_in, false, false, false, false},
+*/
+		{NULL, NULL, NULL, ir_var_auto, nullptr}
+	};
+
+	/** Domain shader system values. */
+	static FSystemValue DomainSystemValueTable[] =
+	{
+/*
+		// TODO : SV_DomainLocation has types float2 or float3 depending on the input topology
+		{"SV_Position", glsl_type::vec4_type, "gl_Position", ir_var_in, false, true, true, false},
+		{"SV_Position", glsl_type::vec4_type, "gl_Position", ir_var_out, false, false, true, false},
+		{"SV_DomainLocation", glsl_type::vec3_type, "gl_TessCoord", ir_var_in, false, false, false, false},
+*/
+		{NULL, NULL, NULL, ir_var_auto, nullptr}
+	};
+
+	/** Compute shader system values. */
+	static FSystemValue ComputeSystemValueTable[] =
+	{
+/*
+		{"SV_DispatchThreadID", glsl_type::uvec3_type, "gl_GlobalInvocationID", ir_var_in, false, false, false, false},
+		{"SV_GroupID", glsl_type::uvec3_type, "gl_WorkGroupID", ir_var_in, false, false, false, false},
+		{"SV_GroupIndex", glsl_type::uint_type, "gl_LocalInvocationIndex", ir_var_in, false, false, false, false},
+		{"SV_GroupThreadID", glsl_type::uvec3_type, "gl_LocalInvocationID", ir_var_in, false, false, false, false},
+*/
+		{NULL, NULL, NULL, ir_var_auto, nullptr}
+	};
+
+	FSystemValue* SystemValueTable[HSF_FrequencyCount] =
+	{
+		VertexSystemValueTable,
+		PixelSystemValueTable,
+		GeometrySystemValueTable,
+		HullSystemValueTable,
+		DomainSystemValueTable,
+		ComputeSystemValueTable
+	};
+
+
+	static ir_rvalue* GenerateInputFromSemantic(EHlslShaderFrequency Frequency, _mesa_glsl_parse_state* ParseState,
+		const char* Semantic, const glsl_type* Type, exec_list* DeclInstructions)
+	{
+		if (!Semantic)
+		{
+			//todo-rco: More info!
+			_mesa_glsl_error(ParseState, "Missing input semantic!", Semantic);
+			return nullptr;
+		}
+
+		ir_variable* Variable = NULL;
+		if (_strnicmp(Semantic, "SV_", 3) == 0)
+		{
+			FSystemValue* SystemValues = SystemValueTable[Frequency];
+			for (int i = 0; SystemValues[i].HlslSemantic != NULL; ++i)
+			{
+				if (SystemValues[i].Mode == ir_var_in && _stricmp(SystemValues[i].HlslSemantic, Semantic) == 0)
+				{
+					ir_variable* Variable = new(ParseState)ir_variable(
+						SystemValues[i].Type, SystemValues[i].MetalName, ir_var_in);
+					Variable->semantic = SystemValues[i].MetalSemantic;
+					Variable->read_only = true;
+					Variable->origin_upper_left = false;
+					DeclInstructions->push_tail(Variable);
+					ParseState->symbols->add_variable(Variable);
+					ir_dereference_variable* VariableDeref = new(ParseState)ir_dereference_variable(Variable);
+
+					return VariableDeref;
+				}
+			}
+		}
+
+		// If we're here, no built-in variables matched.
+		if (_strnicmp(Semantic, "SV_", 3) == 0)
+		{
+			_mesa_glsl_warning(ParseState, "unrecognized system value input '%s'", Semantic);
+		}
+
+		Variable = new(ParseState)ir_variable(
+			Type,
+			ralloc_asprintf(ParseState, "IN_%s", Semantic),
+			ir_var_in);
+		if (Frequency == HSF_VertexShader)
+		{
+			if (!_strnicmp(Semantic, "ATTRIBUTE", 9))
+			{
+				Variable->semantic = ralloc_asprintf(ParseState, "[[ attribute(%s) ]]", Semantic);
+			}
+			else
+			{
+				_mesa_glsl_warning(ParseState, "Unrecognized input attribute '%s'", Semantic);
+			}
+		}
+
+		if (!Variable->semantic)
+		{
+			Variable->semantic = ralloc_asprintf(ParseState, "[[ user(%s) ]]", Semantic);
+		}
+		Variable->read_only = true;
+		DeclInstructions->push_tail(Variable);
+		ParseState->symbols->add_variable(Variable);
+		ir_dereference_variable* VariableDeref = new(ParseState)ir_dereference_variable(Variable);
+		return VariableDeref;
+	}
+
+	static void GenerateInputForVariable(EHlslShaderFrequency Frequency, _mesa_glsl_parse_state* ParseState,
+		const char* InputSemantic, ir_dereference* InputVariableDeref, exec_list* DeclInstructions, exec_list* PreCallInstructions)
+	{
+		const glsl_type* InputType = InputVariableDeref->type;
+		if (InputType->is_record())
+		{
+			for (uint32 i = 0; i < InputType->length; ++i)
+			{
+				const char* Semantic = nullptr;
+				const char* FieldSemantic = InputType->fields.structure[i].semantic;
+				if (InputSemantic && FieldSemantic)
+				{
+					_mesa_glsl_warning(ParseState, "semantic '%s' of field '%s' will be overridden by enclosing types' semantic '%s'",
+						InputType->fields.structure[i].semantic,
+						InputType->fields.structure[i].name,
+						InputSemantic);
+					FieldSemantic = nullptr;
+				}
+				else if (InputSemantic && !FieldSemantic)
+				{
+					Semantic = ralloc_asprintf(ParseState, "%s%d", InputSemantic, i);
+					_mesa_glsl_warning(ParseState, "  creating semantic '%s' for struct field '%s'", Semantic, InputType->fields.structure[i].name);
+				}
+				else if (!InputSemantic && FieldSemantic)
+				{
+					Semantic = FieldSemantic;
+				}
+				else
+				{
+					Semantic = nullptr;
+				}
+
+				if (InputType->fields.structure[i].type->is_record() || Semantic)
+				{
+					ir_dereference_record* FieldDeref = new(ParseState)ir_dereference_record(
+						InputVariableDeref->clone(ParseState, NULL),
+						InputType->fields.structure[i].name);
+					GenerateInputForVariable(Frequency, ParseState, Semantic, FieldDeref, DeclInstructions, PreCallInstructions);
+				}
+				else
+				{
+					_mesa_glsl_error(
+						ParseState,
+						"field '%s' in input structure '%s' does not specify a semantic",
+						InputType->fields.structure[i].name,
+						InputType->name
+						);
+				}
+			}
+		}
+		else
+		{
+			if (InputType->is_array())
+			{
+				int BaseIndex = 0;
+				const char* Semantic = 0;
+				check(InputSemantic);
+				ParseSemanticAndIndex(ParseState, InputSemantic, &Semantic, &BaseIndex);
+				check(BaseIndex >= 0);
+				for (unsigned i = 0; i < InputType->length; ++i)
+				{
+					ir_dereference_array* ArrayDeref = new(ParseState)ir_dereference_array(
+						InputVariableDeref->clone(ParseState, NULL),
+						new(ParseState)ir_constant((unsigned)i)
+						);
+					GenerateInputForVariable(
+						Frequency,
+						ParseState,
+						ralloc_asprintf(ParseState, "%s%d", Semantic, BaseIndex + i),
+						ArrayDeref,
+						DeclInstructions,
+						PreCallInstructions);
+				}
+			}
+			else
+			{
+				ir_rvalue* SrcValue = GenerateInputFromSemantic(Frequency, ParseState, InputSemantic, InputType, DeclInstructions);
+				if (SrcValue)
+				{
+					YYLTYPE loc;
+					apply_type_conversion(InputType, SrcValue, PreCallInstructions, ParseState, true, &loc);
+					PreCallInstructions->push_tail(
+						new(ParseState) ir_assignment(InputVariableDeref->clone(ParseState, NULL),SrcValue));
+				}
+			}
+		}
+	}
+
+	ir_dereference_variable* GenerateInput(EHlslShaderFrequency Frequency, _mesa_glsl_parse_state* ParseState, const char* InputSemantic, const glsl_type* InputType, exec_list* DeclInstructions, exec_list* PreCallInstructions)
+	{
+		ir_variable* TempVariable = new(ParseState)ir_variable(InputType, nullptr, ir_var_temporary);
+		ir_dereference_variable* TempVariableDeref = new(ParseState)ir_dereference_variable(TempVariable);
+		PreCallInstructions->push_tail(TempVariable);
+		GenerateInputForVariable(Frequency, ParseState, InputSemantic, TempVariableDeref, DeclInstructions, PreCallInstructions);
+		return TempVariableDeref;
+	}
+
+	static ir_rvalue* GenerateOutputFromSemantic(EHlslShaderFrequency Frequency, _mesa_glsl_parse_state* ParseState,
+		const char* Semantic, const glsl_type* Type, exec_list* DeclInstructions, const glsl_type** DestVariableType)
+	{
+		ir_variable* Variable = NULL;
+
+		if (!Variable && _strnicmp(Semantic, "SV_", 3) == 0)
+		{
+			FSystemValue* SystemValues = SystemValueTable[Frequency];
+			for (int i = 0; SystemValues[i].HlslSemantic != nullptr; ++i)
+			{
+				if (SystemValues[i].Mode == ir_var_out && _stricmp(SystemValues[i].HlslSemantic, Semantic) == 0)
+				{
+					Variable = new(ParseState) ir_variable(SystemValues[i].Type, SystemValues[i].MetalName, ir_var_out);
+					Variable->semantic = SystemValues[i].MetalSemantic;
+					break;
+				}
+			}
+		}
+
+		if (Semantic && _strnicmp(Semantic, "SV_", 3) == 0 && !Variable)
+		{
+			_mesa_glsl_warning(ParseState, "unrecognized system value output '%s'", Semantic);
+		}
+
+		if (!Variable)
+		{
+			Variable = new(ParseState)ir_variable(Type, ralloc_asprintf(ParseState, "OUT_%s", Semantic), ir_var_out);
+			Variable->semantic = ralloc_asprintf(ParseState, "[[ user(%s) ]]", Semantic);
+		}
+
+		*DestVariableType = Variable->type;
+		DeclInstructions->push_tail(Variable);
+		ParseState->symbols->add_variable(Variable);
+		ir_rvalue* VariableDeref = new(ParseState)ir_dereference_variable(Variable);
+		return VariableDeref;
+	}
+
+	static void GenerateOutputForVariable(EHlslShaderFrequency Frequency, _mesa_glsl_parse_state* ParseState,
+		const char* OutputSemantic, ir_dereference* OutputVariableDeref, exec_list* DeclInstructions, exec_list* PostCallInstructions
+		/*int SemanticArraySize,int SemanticArrayIndex*/)
+	{
+		const glsl_type* OutputType = OutputVariableDeref->type;
+		if (OutputType->is_record())
+		{
+			for (uint32 i = 0; i < OutputType->length; ++i)
+			{
+				const char* FieldSemantic = OutputType->fields.structure[i].semantic;
+				const char* Semantic = nullptr;
+				if (OutputSemantic && FieldSemantic)
+				{
+					_mesa_glsl_warning(ParseState, "semantic '%s' of field '%s' will be overridden by enclosing types' semantic '%s'",
+						OutputType->fields.structure[i].semantic,
+						OutputType->fields.structure[i].name,
+						OutputSemantic);
+					FieldSemantic = nullptr;
+				}
+				else if (OutputSemantic && !FieldSemantic)
+				{
+					Semantic = ralloc_asprintf(ParseState, "%s%d", OutputSemantic, i);
+					_mesa_glsl_warning(ParseState, "  creating semantic '%s' for struct field '%s'", Semantic, OutputType->fields.structure[i].name);
+				}
+				else if (!OutputSemantic && FieldSemantic)
+				{
+					Semantic = FieldSemantic;
+				}
+				else
+				{
+					Semantic = nullptr;
+				}
+
+				if (OutputType->fields.structure[i].type->is_record() || Semantic)
+				{
+					// Dereference the field and generate shader outputs for the field.
+					ir_dereference* FieldDeref = new(ParseState)ir_dereference_record(
+						OutputVariableDeref->clone(ParseState, NULL),
+						OutputType->fields.structure[i].name);
+					GenerateOutputForVariable(Frequency, ParseState, Semantic, FieldDeref, DeclInstructions, PostCallInstructions);
+				}
+				else
+				{
+					_mesa_glsl_error(
+						ParseState,
+						"field '%s' in output structure '%s' does not specify a semantic",
+						OutputType->fields.structure[i].name,
+						OutputType->name
+						);
+				}
+			}
+		}
+		else
+		{
+			if (!OutputSemantic)
+			{
+				_mesa_glsl_error(ParseState, "Entry point does not specify a semantic for its return value");
+			}
+			else
+			{
+				if (OutputType->is_array())
+				{
+					int BaseIndex = 0;
+					const char* Semantic = 0;
+
+					ParseSemanticAndIndex(ParseState, OutputSemantic, &Semantic, &BaseIndex);
+
+					for (unsigned i = 0; i < OutputType->length; ++i)
+					{
+						ir_dereference_array* ArrayDeref = new(ParseState)ir_dereference_array(
+							OutputVariableDeref->clone(ParseState, NULL),
+							new(ParseState) ir_constant((unsigned)i)
+							);
+						GenerateOutputForVariable(Frequency, ParseState,
+							ralloc_asprintf(ParseState, "%s%d", Semantic, BaseIndex + i),
+							ArrayDeref, DeclInstructions, PostCallInstructions);
+					}
+				}
+				else
+				{
+					YYLTYPE loc;
+					ir_rvalue* Src = OutputVariableDeref->clone(ParseState, NULL);
+					const glsl_type* DestVariableType = NULL;
+					ir_rvalue* DestVariableDeref = GenerateOutputFromSemantic(Frequency, ParseState, OutputSemantic,
+						OutputType, DeclInstructions, &DestVariableType);
+
+					apply_type_conversion(DestVariableType, Src, PostCallInstructions, ParseState, true, &loc);
+					PostCallInstructions->push_tail(new(ParseState)ir_assignment(DestVariableDeref, Src));
+				}
+			}
+		}
+	}
+
+	ir_dereference_variable* GenerateOutput(EHlslShaderFrequency Frequency, _mesa_glsl_parse_state* ParseState,
+		const char* OutputSemantic, const glsl_type* OutputType, exec_list* DeclInstructions, exec_list* PreCallInstructions, exec_list* PostCallInstructions)
+	{
+		// Generate a local variable to hold the output.
+		ir_variable* TempVariable = new(ParseState) ir_variable(OutputType, nullptr, ir_var_temporary);
+		ir_dereference_variable* TempVariableDeref = new(ParseState) ir_dereference_variable(TempVariable);
+		PreCallInstructions->push_tail(TempVariable);
+
+		GenerateOutputForVariable(Frequency, ParseState, OutputSemantic, TempVariableDeref, DeclInstructions, PostCallInstructions);
+
+		return TempVariableDeref;
+	}
+}
+
+struct FFixIntrinsicsVisitor : public ir_rvalue_visitor
+{
+	_mesa_glsl_parse_state* State;
+	bool bUsesFramebufferFetchES2;
+	int MRTFetchMask;
+	ir_variable* DestColorVar;
+	const glsl_type* DestColorType;
+	ir_variable* DestMRTColorVar[MAX_SIMULTANEOUS_RENDER_TARGETS];
+
+	FFixIntrinsicsVisitor(_mesa_glsl_parse_state* InState, ir_function_signature* InMainSig) :
+		State(InState),
+		bUsesFramebufferFetchES2(false),
+		MRTFetchMask(0),
+		DestColorVar(nullptr),
+		DestColorType(glsl_type::error_type)
+	{
+		DestColorType = GetFragColorTypeFromMetalOutputStruct(InMainSig->return_type);
+		memset(DestMRTColorVar, 0, sizeof(DestMRTColorVar));
+	}
+
+	//ir_visitor_status visit_leave(ir_expression* expr) override
+	virtual void handle_rvalue(ir_rvalue** RValue)
+	{
+		if (!RValue || !*RValue)
+		{
+			return;
+		}
+		auto* expr = (*RValue)->as_expression();
+		if (!expr)
+		{
+			return;
+		}
+
+		ir_expression_operation op = expr->operation;
+
+		if (op == ir_binop_mul && expr->type->is_matrix()
+			&& expr->operands[0]->type->is_matrix()
+			&& expr->operands[1]->type->is_matrix())
+		{
+			// Convert matrixCompMult to memberwise multiply
+			check(expr->operands[0]->type == expr->operands[1]->type);
+			auto* NewTemp = new(State)ir_variable(expr->operands[0]->type, nullptr, ir_var_temporary);
+			base_ir->insert_before(NewTemp);
+			for (uint32 Index = 0; Index < expr->operands[0]->type->matrix_columns; ++Index)
+			{
+				auto* NewMul = new(State)ir_expression(ir_binop_mul,
+					new(State)ir_dereference_array(expr->operands[0], new(State)ir_constant(Index)),
+					new(State)ir_dereference_array(expr->operands[1], new(State)ir_constant(Index)));
+				auto* NewAssign = new(State)ir_assignment(
+					new(State)ir_dereference_array(NewTemp, new(State)ir_constant(Index)),
+					NewMul);
+				base_ir->insert_before(NewAssign);
+			}
+
+			*RValue = new(State)ir_dereference_variable(NewTemp);
+		}
+	}
+
+	virtual ir_visitor_status visit_leave(ir_call* IR) override
+	{
+		if (IR->use_builtin)
+		{
+			const char* CalleeName = IR->callee_name();
+			static auto ES2Len = strlen(FRAMEBUFFER_FETCH_ES2);
+			static auto MRTLen = strlen(FRAMEBUFFER_FETCH_MRT);
+			if (!strncmp(CalleeName, FRAMEBUFFER_FETCH_ES2, ES2Len))
+			{
+				// 'Upgrade' framebuffer fetch
+				check(IR->actual_parameters.is_empty());
+				bUsesFramebufferFetchES2 = true;
+				if (!DestColorVar)
+				{
+					// Generate new input variable for Metal semantics
+					DestColorVar = new(State)ir_variable(glsl_type::get_instance(DestColorType->base_type, 4, 1), "gl_LastFragData", ir_var_in);
+					DestColorVar->semantic = "[[ color(0) ]]";
+				}
+
+				ir_rvalue* DestColor = new(State)ir_dereference_variable(DestColorVar);
+				if (IR->return_deref->type->base_type != DestColor->type->base_type)
+				{
+					DestColor = convert_component(DestColor, IR->return_deref->type);
+				}
+				auto* Assignment = new (State)ir_assignment(IR->return_deref, DestColor);
+				IR->insert_before(Assignment);
+				IR->remove();
+			}
+			else if (!strncmp(CalleeName, FRAMEBUFFER_FETCH_MRT, MRTLen))
+			{
+				int Index = atoi(CalleeName + MRTLen);
+				if (!DestMRTColorVar[Index])
+				{
+					DestMRTColorVar[Index] = new(State)ir_variable(glsl_type::get_instance(DestColorType->base_type, 4, 1), CalleeName, ir_var_in);
+					DestMRTColorVar[Index]->semantic = ralloc_asprintf(State, "[[ color(%d) ]]", Index);
+				}
+
+				ir_rvalue* DestColor = new(State) ir_dereference_variable(DestMRTColorVar[Index]);
+				if (IR->return_deref->type->base_type != DestColor->type->base_type)
+				{
+					DestColor = convert_component(DestColor, IR->return_deref->type);
+				}
+				auto* Assignment = new (State) ir_assignment(IR->return_deref, DestColor);
+				IR->insert_before(Assignment);
+				IR->remove();
+			}
+		}
+
+		return visit_continue;
+	}
+};
+
+void FMetalCodeBackend::FixIntrinsics(exec_list* ir, _mesa_glsl_parse_state* state)
+{
+	ir_function_signature* MainSig = GetMainFunction(ir);
+	check(MainSig);
+
+	FFixIntrinsicsVisitor Visitor(state,MainSig);
+	Visitor.run(&MainSig->body);
+
+	if (Visitor.bUsesFramebufferFetchES2)
+	{
+		check(Visitor.DestColorVar);
+		MainSig->parameters.push_tail(Visitor.DestColorVar);
+	}
+
+	for (int i = 0; i < MAX_SIMULTANEOUS_RENDER_TARGETS; ++i)
+	{
+		if (Visitor.DestMRTColorVar[i])
+		{
+			MainSig->parameters.push_tail(Visitor.DestMRTColorVar[i]);
+		}
+	}
+}
+
+struct FConvertUBVisitor : public ir_rvalue_visitor
+{
+	_mesa_glsl_parse_state* State;
+	TStringIRVarMap& Map;
+	FConvertUBVisitor(_mesa_glsl_parse_state* InState, TStringIRVarMap& InMap) :
+		State(InState),
+		Map(InMap)
+	{
+	}
+
+	virtual void handle_rvalue(ir_rvalue** RValuePtr) override
+	{
+		if (!RValuePtr || !*RValuePtr)
+		{
+			return;
+		}
+		auto* ReferencedVar = (*RValuePtr)->variable_referenced();
+		if (ReferencedVar && ReferencedVar->mode == ir_var_uniform && ReferencedVar->semantic)
+		{
+			auto FoundIter = Map.find(ReferencedVar->semantic);
+			if (FoundIter != Map.end())
+			{
+				auto* StructVar = FoundIter->second;
+				StructVar->used = 1;
+
+				// Actually replace the variable
+				auto* DeRefVar = (*RValuePtr)->as_dereference_variable();
+				if (DeRefVar)
+				{
+					*RValuePtr = new(State)ir_dereference_record(StructVar, ReferencedVar->name);
+				}
+				else
+				{
+					check(0);
+				}
+			}
+		}
+	}
+};
+
+void FMetalCodeBackend::MovePackedUniformsToMain(exec_list* ir, _mesa_glsl_parse_state* state, FBuffers& OutBuffers)
+{
+	//IRDump(ir);
+	TStringIRVarMap CBVarMap;
+
+	// Now make a new struct type and global variable per uniform buffer
+	for (uint32 i = 0; i < state->num_uniform_blocks; ++i)
+		//	for (auto& CB : state->CBuffersOriginal)
+	{
+		auto* CBP = state->FindCBufferByName(false, state->uniform_blocks[i]->name);
+		check(CBP);
+		auto& CB = *CBP;
+		if (!CB.Members.empty())
+		{
+			glsl_struct_field* Fields = ralloc_array(state, glsl_struct_field, (unsigned)CB.Members.size());
+			uint32 Index = 0;
+			for (auto& Member : CB.Members)
+			{
+				check(Member.Var);
+				Fields[Index++] = glsl_struct_field(Member.Var->type, ralloc_strdup(state, Member.Var->name));
+			}
+
+			auto* Type = glsl_type::get_record_instance(Fields, (unsigned)CB.Members.size(), ralloc_asprintf(state, "CB_%s", CB.Name.c_str()));
+			// Hack: This way we tell this is a uniform buffer and we need to emit 'packed_'
+			((glsl_type*)Type)->HlslName = "__PACKED__";
+			state->AddUserStruct(Type);
+
+			auto* Var = new(state)ir_variable(Type, ralloc_asprintf(state, "%s", CB.Name.c_str()), ir_var_uniform);
+			CBVarMap[CB.Name] = Var;
+		}
+	}
+
+	FConvertUBVisitor ConvertVisitor(state, CBVarMap);
+	ConvertVisitor.run(ir);
+
+	std::set<const glsl_type*> PendingTypes;
+	std::set<const glsl_type*> ProcessedTypes;
+
+	// Actually only save the used variables
+	for (auto& Pair : CBVarMap)
+	{
+		auto* Var = Pair.second;
+		if (Var->used)
+		{
+			// Go through each struct type and mark it as packed
+			ir->push_head(Var);
+			if (Var->type->is_record())
+			{
+				PendingTypes.insert(Var->type);
+			}
+		}
+	}
+
+	// Mark all structures as packed
+	while (!PendingTypes.empty())
+	{
+		auto* Type = *PendingTypes.begin();
+		PendingTypes.erase(PendingTypes.begin());
+		if (ProcessedTypes.find(Type) == ProcessedTypes.end())
+		{
+			ProcessedTypes.insert(Type);
+			((glsl_type*)Type)->HlslName = "__PACKED__";
+
+			for (uint32 i = 0; i < Type->length; ++i)
+			{
+				if (Type->fields.structure[i].type->is_record())
+				{
+					PendingTypes.insert(Type->fields.structure[i].type);
+				}
+			}
+		}
+	}
+
+	ir_function_signature* MainSig = GetMainFunction(ir);
+	check(MainSig);
+
+	// Gather all globals still lying outside Main
+	foreach_iter(exec_list_iterator, iter, *ir)
+	{
+		auto* Instruction = ((ir_instruction*)iter.get());
+		auto* Var = Instruction->as_variable();
+		if (Var)
+		{
+			check(Var->mode == ir_var_uniform || Var->mode == ir_var_out || Var->mode == ir_var_in);
+			OutBuffers.Buffers.Add(Var);
+		}
+	}
+	OutBuffers.SortBuffers();
+
+	// And move them to main
+	for (auto Iter : OutBuffers.Buffers)
+	{
+		auto* Var = (ir_variable*)Iter;
+		if (Var)
+		{
+			Var->remove();
+			MainSig->parameters.push_tail(Var);
+		}
+	}
+	//IRDump(ir, state);
+}
+
+
+void FMetalCodeBackend::PromoteInputsAndOutputsGlobalHalfToFloat(exec_list* Instructions, _mesa_glsl_parse_state* State, EHlslShaderFrequency Frequency)
+{
+	//IRDump(Instructions);
+	ir_function_signature* EntryPointSig = GetMainFunction(Instructions);
+	check(EntryPointSig);
+	foreach_iter(exec_list_iterator, Iter, *Instructions)
+	{
+		ir_instruction* IR = (ir_instruction*)Iter.get();
+		auto* Variable = IR->as_variable();
+		if (Variable)
+		{
+			auto* NewType = PromoteHalfToFloatType(State, Variable->type);
+			if (Variable->type == NewType)
+			{
+				continue;
+			}
+
+			switch (Variable->mode)
+			{
+			case ir_var_in:
+			{
+							  auto* NewVar = new(State)ir_variable(NewType, Variable->name, ir_var_in);
+							  NewVar->semantic = Variable->semantic;
+							  Variable->insert_before(NewVar);
+							  Variable->name = nullptr;
+							  Variable->semantic = nullptr;
+							  Variable->mode = ir_var_temporary;
+							  Variable->remove();
+							  exec_list Assignments;
+							  Assignments.push_head(Variable);
+							  CreateNewAssignmentsFloat2Half(State, Assignments, Variable, new(State)ir_dereference_variable(NewVar));
+							  EntryPointSig->body.get_head()->insert_before(&Assignments);
+							  break;
+			}
+			case ir_var_out:
+			{
+							   if (Frequency != HSF_PixelShader)
+							   {
+								   auto* NewVar = new(State)ir_variable(NewType, Variable->name, ir_var_out);
+								   NewVar->semantic = Variable->semantic;
+								   Variable->insert_before(NewVar);
+								   Variable->name = nullptr;
+								   Variable->semantic = nullptr;
+								   Variable->mode = ir_var_temporary;
+								   Variable->remove();
+								   exec_list Assignments;
+								   CreateNewAssignmentsHalf2Float(State, Assignments, NewVar, new(State)ir_dereference_variable(Variable));
+								   EntryPointSig->body.push_head(Variable);
+								   EntryPointSig->body.append_list(&Assignments);
+								   break;
+							   }
+			}
+			}
+		}
+	}
+	//IRDump(Instructions);
+}
+
+static bool ProcessStageInVariables(_mesa_glsl_parse_state* ParseState, EHlslShaderFrequency Frequency, ir_variable* Variable, TArray<glsl_struct_field>& OutStageInMembers, std::set<ir_variable*>& OutStageInVariables, unsigned int* OutVertexAttributesMask, TIRVarList& OutFunctionArguments)
+{
+	// Don't move variables that are system values into the input structures
+	const auto* SystemValues = MetalUtils::SystemValueTable[Frequency];
+	for(int i = 0; SystemValues[i].MetalSemantic != nullptr; ++i)
+	{
+		if (SystemValues[i].Mode == ir_var_in && _stricmp(SystemValues[i].MetalSemantic, Variable->semantic) == 0)
+		{
+			return true;
+		}
+	}
+
+
+	if (Frequency == HSF_VertexShader)
+	{
+		// Generate an uber struct
+		if (Variable->type->is_record())
+		{
+			check(0);
+#if 0
+			// Make sure all members have semantics
+			std::list<const glsl_type*> PendingTypes;
+			std::set<const glsl_type*> VisitedTypes;
+			PendingTypes.push_back(Variable->type);
+			while (!PendingTypes.empty())
+			{
+				auto* Type = PendingTypes.front();
+				PendingTypes.pop_front();
+				if (VisitedTypes.find(Type) == VisitedTypes.end())
+				{
+					VisitedTypes.insert(Type);
+					bool bHasSemantic = true;
+					for (int i = 0; i < Type->length; ++i)
+					{
+						auto& Member = Type->fields.structure[i];
+						if (Member.type->is_record())
+						{
+							PendingTypes.push_back(Member.type);
+						}
+						else
+						{
+							if (Member.semantic)
+							{
+								int AttributeIndex = GetAttributeIndex(Member.semantic);
+								if (AttributeIndex < 0)
+								{
+									_mesa_glsl_error(ParseState, "struct %s has a member with unknown semantic %s!\n", Type->name, Member.semantic);
+									return false;
+								}
+								if (Member.type->is_array())
+								{
+									check(Member.type->element_type()->is_vector());
+									for (int i = 0; i < Member.type->length; ++i, ++AttributeIndex)
+									{
+										glsl_struct_field OutMember;
+										OutMember.type = Member.type->element_type();
+										OutMember.semantic = ralloc_asprintf(ParseState, "ATTRIBUTE%d", AttributeIndex);
+										OutMember.name = ralloc_asprintf(ParseState, "Attribute%d", AttributeIndex);
+
+										if (OutVertexAttributesMask)
+										{
+											*OutVertexAttributesMask |= (1 << AttributeIndex);
+										}
+
+										OutStageInMembers.push_back(OutMember);
+									}
+								}
+								else
+								{
+									glsl_struct_field OutMember;
+									OutMember.type = Member.type;
+									OutMember.semantic = ralloc_strdup(ParseState, Member.semantic);
+									OutMember.name = ralloc_asprintf(ParseState, "Attribute%d", AttributeIndex);
+
+									if (OutVertexAttributesMask)
+									{
+										*OutVertexAttributesMask |= (1 << AttributeIndex);
+									}
+
+									OutStageInMembers.push_back(OutMember);
+								}
+							}
+							else
+							{
+								bHasSemantic = false;
+								break;
+							}
+						}
+						if (!bHasSemantic)
+						{
+							_mesa_glsl_error(ParseState, "struct %s has a member with no semantic!\n", Type->name);
+							return false;
+						}
+					}
+				}
+			}
+#endif
+		}
+		else
+		{
+			int AttributeIndex = GetInAttributeIndex(Variable->name);
+			if (AttributeIndex >= 0)
+			{
+				if (Variable->type->is_array())
+				{
+					check(Variable->type->element_type()->is_vector());
+					for (uint32 i = 0; i < Variable->type->length; ++i, ++AttributeIndex)
+					{
+						glsl_struct_field OutMember;
+						OutMember.type = Variable->type->element_type();
+						OutMember.semantic = ralloc_asprintf(ParseState, "ATTRIBUTE%d", AttributeIndex);
+						OutMember.name = ralloc_asprintf(ParseState, "ATTRIBUTE%d", AttributeIndex);
+
+						if (OutVertexAttributesMask)
+						{
+							*OutVertexAttributesMask |= (1 << AttributeIndex);
+						}
+
+						OutStageInMembers.Add(OutMember);
+					}
+				}
+				else
+				{
+					glsl_struct_field OutMember;
+					OutMember.type = Variable->type;
+					OutMember.semantic = ralloc_asprintf(ParseState, "ATTRIBUTE%d", AttributeIndex);
+					OutMember.name = ralloc_asprintf(ParseState, "IN_ATTRIBUTE%d", AttributeIndex);
+
+					if (OutVertexAttributesMask)
+					{
+						*OutVertexAttributesMask |= (1 << AttributeIndex);
+					}
+
+					OutStageInMembers.Add(OutMember);
+				}
+			}
+			else if (!strcmp(Variable->name, "gl_VertexID") || !strcmp(Variable->name, "gl_InstanceID"))
+			{
+				OutFunctionArguments.push_back(Variable);
+				return true;
+			}
+			else
+			{
+				_mesa_glsl_error(ParseState, "Unknown semantic for input attribute %s!\n", Variable->semantic ? Variable->semantic : "", Variable->name);
+				check(0);
+				return false;
+			}
+		}
+
+		OutStageInVariables.insert(Variable);
+		return true;
+	}
+	else
+	{
+		check(Frequency == HSF_PixelShader);
+		if (!strcmp(Variable->name, "gl_FrontFacing"))
+		{
+			// Make sure we add a semantic
+			Variable->semantic = "gl_FrontFacing";
+			return true;
+		}
+	}
+
+	/*
+	if (!VerifyVariableHasSemantics(ParseState, Variable))
+	{
+	return false;
+	}
+	*/
+
+	glsl_struct_field Member;
+	Member.type = Variable->type;
+	Member.name = ralloc_strdup(ParseState, Variable->name);
+	Member.semantic = ralloc_strdup(ParseState, Variable->semantic ? Variable->semantic : Variable->name);
+	OutStageInMembers.Add(Member);
+	OutStageInVariables.insert(Variable);
+
+	return true;
+}
+
+
+/** Information on system values. */
+struct FSystemValue
+{
+	const char* Semantic;
+	const glsl_type* Type;
+	const char* GlslName;
+	ir_variable_mode Mode;
+	bool bOriginUpperLeft;
+	bool bArrayVariable;
+};
+
+/** Vertex shader system values. */
+static FSystemValue VertexSystemValueTable[] =
+{
+	{"SV_VertexID", glsl_type::uint_type, "gl_VertexID", ir_var_in, false, false},
+	{"SV_InstanceID", glsl_type::uint_type, "gl_InstanceID", ir_var_in, false, false},
+	//{ "SV_Position", glsl_type::vec4_type, "gl_Position", ir_var_out, false, true },
+	{NULL, NULL, NULL, ir_var_auto, false, false}
+};
+
+/** Pixel shader system values. */
+static FSystemValue PixelSystemValueTable[] =
+{
+	{"SV_Depth", glsl_type::float_type, "gl_FragDepth", ir_var_out, false, false},
+	{"SV_Position", glsl_type::vec4_type, "gl_FragCoord", ir_var_in, true, false},
+	//	{ "SV_IsFrontFace", glsl_type::bool_type, "gl_FrontFacing", ir_var_in, false, true },
+	{"SV_PrimitiveID", glsl_type::int_type, "gl_PrimitiveID", ir_var_in, false, false},
+	//	{ "SV_RenderTargetArrayIndex", glsl_type::int_type, "gl_Layer", ir_var_in, false, false },
+	//	{ "SV_Target0", glsl_type::vec4_type, "gl_FragColor", ir_var_out, false, true },
+	{NULL, NULL, NULL, ir_var_auto, false, false}
+};
+
+static FSystemValue* SystemValueTable[] =
+{
+	VertexSystemValueTable,
+	PixelSystemValueTable,
+	nullptr,
+	nullptr,
+	nullptr,
+	nullptr
+};
+
+/**
+* Generate an input semantic.
+* @param Frequency - The shader frequency.
+* @param ParseState - Parse state.
+* @param Semantic - The semantic name to generate.
+* @param Type - Value type.
+* @param DeclInstructions - IR to which declarations may be added.
+* @returns reference to IR variable for the semantic.
+*/
+static ir_rvalue* GenShaderInputSemantic(
+	EHlslShaderFrequency Frequency,
+	_mesa_glsl_parse_state* ParseState,
+	const char* Semantic,
+	const glsl_type* Type,
+	exec_list* DeclInstructions,
+	int SemanticArraySize,
+	int SemanticArrayIndex)
+{
+	ir_variable* Variable = NULL;
+	if (Semantic && _strnicmp(Semantic, "SV_", 3) == 0)
+	{
+		FSystemValue* SystemValues = SystemValueTable[Frequency];
+		for (int i = 0; SystemValues[i].Semantic != NULL; ++i)
+		{
+			if (SystemValues[i].Mode == ir_var_in
+				//				&& (!SystemValues[i].bESOnly || ParseState->bGenerateES)
+				&& _stricmp(SystemValues[i].Semantic, Semantic) == 0)
+			{
+				check(0);
+#if 0
+				if (SystemValues[i].bArrayVariable)
+				{
+					// Built-in array variable. Like gl_in[x].gl_Position.
+					// The variable for it has already been created in GenShaderInput().
+					/*ir_variable**/ Variable = ParseState->symbols->get_variable("gl_in");
+					check(Variable);
+					ir_dereference_variable* ArrayDeref = new(ParseState)ir_dereference_variable(Variable);
+					ir_dereference_array* StructDeref = new(ParseState)ir_dereference_array(
+						ArrayDeref,
+						new(ParseState)ir_constant((unsigned)SemanticArrayIndex)
+						);
+					ir_dereference_record* VariableDeref = new(ParseState)ir_dereference_record(
+						StructDeref,
+						SystemValues[i].GlslName
+						);
+					// TO DO - in case of SV_ClipDistance, we need to defer appropriate index in variable too.
+					return VariableDeref;
+				}
+				else
+				{
+					// Built-in variable that shows up only once, like gl_FragCoord in fragment
+					// shader, or gl_PrimitiveIDIn in geometry shader. Unlike gl_in[x].gl_Position.
+					// Even in geometry shader input pass it shows up only once.
+
+					// Create it on first pass, ignore the call on others.
+					if (SemanticArrayIndex == 0)
+					{
+						ir_variable* Variable = new(ParseState)ir_variable(
+							SystemValues[i].Type,
+							SystemValues[i].GlslName,
+							ir_var_in
+							);
+						Variable->read_only = true;
+						Variable->origin_upper_left = SystemValues[i].bOriginUpperLeft;
+						DeclInstructions->push_tail(Variable);
+						ParseState->symbols->add_variable(Variable);
+						ir_dereference_variable* VariableDeref = new(ParseState)ir_dereference_variable(Variable);
+
+						return VariableDeref;
+					}
+					else
+					{
+						return NULL;
+					}
+				}
+#endif
+			}
+		}
+	}
+
+	if (Variable)
+	{
+		// Up to this point, variables aren't contained in structs
+		DeclInstructions->push_tail(Variable);
+		ParseState->symbols->add_variable(Variable);
+		Variable->centroid = false;// InputQualifier.Fields.bCentroid;
+		Variable->interpolation = false;//InputQualifier.Fields.InterpolationMode;
+		Variable->is_patch_constant = false;//InputQualifier.Fields.bIsPatchConstant;
+		ir_rvalue* VariableDeref = new(ParseState)ir_dereference_variable(Variable);
+
+		return VariableDeref;
+	}
+
+	// If we're here, no built-in variables matched.
+
+	if (Semantic && _strnicmp(Semantic, "SV_", 3) == 0)
+	{
+		_mesa_glsl_warning(ParseState, "unrecognized system "
+			"value input '%s'", Semantic);
+	}
+
+	if (Frequency == HSF_VertexShader || ParseState->bGenerateES)
+	{
+		const char* Prefix = "in";
+		if (ParseState->bGenerateES && Frequency == HSF_PixelShader)
+		{
+			Prefix = "var";
+		}
+
+		// Vertex shader inputs don't get packed into structs that we'll later morph into interface blocks
+		if (ParseState->bGenerateES && Type->is_integer())
+		{
+			// Convert integer attributes to floats
+			ir_variable* Variable = new(ParseState)ir_variable(
+				Type,
+				ralloc_asprintf(ParseState, "%s_%s_I", Prefix, Semantic),
+				ir_var_temporary
+				);
+			Variable->centroid = false;//InputQualifier.Fields.bCentroid;
+			Variable->interpolation = false;//InputQualifier.Fields.InterpolationMode;
+			check(Type->is_vector() || Type->is_scalar());
+			check(Type->base_type == GLSL_TYPE_INT || Type->base_type == GLSL_TYPE_UINT);
+
+			// New float attribute
+			ir_variable* ReplacedAttributeVar = new (ParseState)ir_variable(glsl_type::get_instance(GLSL_TYPE_FLOAT, Variable->type->vector_elements, 1), ralloc_asprintf(ParseState, "%s_%s", Prefix, Semantic), ir_var_in);
+			ReplacedAttributeVar->read_only = true;
+			ReplacedAttributeVar->centroid = false;//InputQualifier.Fields.bCentroid;
+			ReplacedAttributeVar->interpolation = false;//InputQualifier.Fields.InterpolationMode;
+
+			// Convert to integer
+			ir_assignment* ConversionAssignment = new(ParseState)ir_assignment(
+				new(ParseState)ir_dereference_variable(Variable),
+				new(ParseState)ir_expression(
+				Type->base_type == GLSL_TYPE_INT ? ir_unop_f2i : ir_unop_f2u,
+				new (ParseState)ir_dereference_variable(ReplacedAttributeVar)
+				)
+				);
+
+			DeclInstructions->push_tail(ReplacedAttributeVar);
+			DeclInstructions->push_tail(Variable);
+			DeclInstructions->push_tail(ConversionAssignment);
+			ParseState->symbols->add_variable(Variable);
+			ParseState->symbols->add_variable(ReplacedAttributeVar);
+
+			ir_dereference_variable* VariableDeref = new(ParseState)ir_dereference_variable(ReplacedAttributeVar);
+			return VariableDeref;
+		}
+
+		// Regular attribute
+		Variable = new(ParseState)ir_variable(
+			Type,
+			ralloc_asprintf(ParseState, "%s_%s", Prefix, Semantic),
+			ir_var_in
+			);
+		Variable->read_only = true;
+		Variable->centroid = false;//InputQualifier.Fields.bCentroid;
+		Variable->interpolation = false;//InputQualifier.Fields.InterpolationMode;
+		Variable->is_patch_constant = false;//InputQualifier.Fields.bIsPatchConstant;
+
+		DeclInstructions->push_tail(Variable);
+		ParseState->symbols->add_variable(Variable);
+
+		ir_dereference_variable* VariableDeref = new(ParseState)ir_dereference_variable(Variable);
+		return VariableDeref;
+	}
+	else if (SemanticArrayIndex == 0)
+	{
+		// On first pass, create variable
+
+		glsl_struct_field *StructField = ralloc_array(ParseState, glsl_struct_field, 1);
+
+		memset(StructField, 0, sizeof(glsl_struct_field));
+		StructField[0].type = Type;
+		StructField[0].name = ralloc_strdup(ParseState, "Data");
+
+		const glsl_type* VariableType = glsl_type::get_record_instance(StructField, 1, ralloc_strdup(ParseState, Semantic));
+		if (SemanticArraySize)
+		{
+			// Pack it into an array too
+			VariableType = glsl_type::get_array_instance(VariableType, SemanticArraySize);
+		}
+
+		ir_variable* Variable = new(ParseState)ir_variable(VariableType, ralloc_asprintf(ParseState, "in_%s", Semantic), ir_var_in);
+		Variable->read_only = true;
+		Variable->is_interface_block = true;
+		Variable->centroid = false;//InputQualifier.Fields.bCentroid;
+		Variable->interpolation = false;//InputQualifier.Fields.InterpolationMode;
+		Variable->is_patch_constant = false;//InputQualifier.Fields.bIsPatchConstant;
+		DeclInstructions->push_tail(Variable);
+		ParseState->symbols->add_variable(Variable);
+
+		ir_rvalue* VariableDeref = new(ParseState)ir_dereference_variable(Variable);
+		if (SemanticArraySize)
+		{
+			// Deref inside array first
+			VariableDeref = new(ParseState)ir_dereference_array(VariableDeref, new(ParseState)ir_constant((unsigned)SemanticArrayIndex)
+				);
+		}
+		VariableDeref = new(ParseState)ir_dereference_record(VariableDeref, ralloc_strdup(ParseState, "Data"));
+		return VariableDeref;
+	}
+	else
+	{
+		// Array variable, not first pass. It already exists, get it.
+		ir_variable* Variable = ParseState->symbols->get_variable(ralloc_asprintf(ParseState, "in_%s", Semantic));
+		check(Variable);
+
+		ir_rvalue* VariableDeref = new(ParseState)ir_dereference_variable(Variable);
+		VariableDeref = new(ParseState)ir_dereference_array(VariableDeref, new(ParseState)ir_constant((unsigned)SemanticArrayIndex));
+		VariableDeref = new(ParseState)ir_dereference_record(VariableDeref, ralloc_strdup(ParseState, "Data"));
+		return VariableDeref;
+	}
+}
+
+/**
+* Generate a shader input.
+* @param Frequency - The shader frequency.
+* @param ParseState - Parse state.
+* @param InputSemantic - The semantic name to generate.
+* @param InputQualifier - Qualifiers applied to the semantic.
+* @param InputType - Value type.
+* @param DeclInstructions - IR to which declarations may be added.
+* @param PreCallInstructions - IR to which instructions may be added before the
+*                              entry point is called.
+* @returns the IR variable deref for the semantic.
+*/
+static ir_dereference_variable* GenerateShaderInput(
+	EHlslShaderFrequency Frequency,
+	_mesa_glsl_parse_state* ParseState,
+	const char* InputSemantic,
+	const glsl_type* InputType,
+	exec_list* DeclInstructions,
+	exec_list* PreCallInstructions)
+{
+	ir_variable* TempVariable = new(ParseState)ir_variable(
+		InputType,
+		NULL,
+		ir_var_temporary);
+	ir_dereference_variable* TempVariableDeref = new(ParseState) ir_dereference_variable(TempVariable);
+	PreCallInstructions->push_tail(TempVariable);
+
+	check(!InputType->is_inputpatch() && !InputType->is_outputpatch());
+	ir_rvalue* SrcValue = MetalUtils::GenerateInputFromSemantic(Frequency, ParseState, InputSemantic, InputType, DeclInstructions);
+	if(SrcValue)
+	{
+		YYLTYPE loc ={0};
+		apply_type_conversion(InputType,SrcValue,PreCallInstructions,ParseState,true,&loc);
+		PreCallInstructions->push_tail(
+			new(ParseState)ir_assignment(
+			TempVariableDeref->clone(ParseState,NULL),
+			SrcValue
+			)
+			);
+	}
+
+	return TempVariableDeref;
+}
+
+
+/**
+* Generate an output semantic.
+* @param Frequency - The shader frequency.
+* @param ParseState - Parse state.
+* @param Semantic - The semantic name to generate.
+* @param Type - Value type.
+* @param DeclInstructions - IR to which declarations may be added.
+* @returns the IR variable for the semantic.
+*/
+static ir_rvalue* GenShaderOutputSemantic(
+	EHlslShaderFrequency Frequency,
+	_mesa_glsl_parse_state* ParseState,
+	const char* Semantic,
+	const glsl_type* Type,
+	exec_list* DeclInstructions,
+	const glsl_type** DestVariableType)
+{
+	FSystemValue* SystemValues = SystemValueTable[Frequency];
+	ir_variable* Variable = NULL;
+
+	if (Semantic && _strnicmp(Semantic, "SV_", 3) == 0)
+	{
+		for (int i = 0; SystemValues[i].Semantic != NULL; ++i)
+		{
+			if (SystemValues[i].Mode == ir_var_out
+				&& _stricmp(SystemValues[i].Semantic, Semantic) == 0)
+			{
+				check(0);
+			}
+		}
+	}
+
+	if (Variable == NULL && Frequency == HSF_VertexShader)
+	{
+		const int PrefixLength = 15;
+		if (_strnicmp(Semantic, "SV_ClipDistance", PrefixLength) == 0
+			&& Semantic[PrefixLength] >= '0'
+			&& Semantic[PrefixLength] <= '9')
+		{
+			check(0);
+		}
+	}
+
+	if (Variable == NULL && Frequency == HSF_PixelShader)
+	{
+		const int PrefixLength = 9;
+		if (_strnicmp(Semantic, "SV_Target", PrefixLength) == 0
+			&& Semantic[PrefixLength] >= '0'
+			&& Semantic[PrefixLength] <= '7')
+		{
+			int OutputIndex = Semantic[PrefixLength] - '0';
+			Variable = new(ParseState)ir_variable(
+				Type,
+				ralloc_asprintf(ParseState, "out_Target%d", OutputIndex),
+				ir_var_out
+				);
+		}
+	}
+
+	if (Variable == NULL && ParseState->bGenerateES)
+	{
+		check(0);
+		// Create a variable so that a struct will not get added
+		Variable = new(ParseState)ir_variable(Type, ralloc_asprintf(ParseState, "var_%s", Semantic), ir_var_out);
+	}
+
+	if (Variable)
+	{
+		// Up to this point, variables aren't contained in structs
+		*DestVariableType = Variable->type;
+		DeclInstructions->push_tail(Variable);
+		ParseState->symbols->add_variable(Variable);
+		Variable->centroid = false;//OutputQualifier.Fields.bCentroid;
+		Variable->interpolation = false;//OutputQualifier.Fields.InterpolationMode;
+		Variable->is_patch_constant = false;// OutputQualifier.Fields.bIsPatchConstant;
+		ir_rvalue* VariableDeref = new(ParseState)ir_dereference_variable(Variable);
+		return VariableDeref;
+	}
+
+	if (Semantic && _strnicmp(Semantic, "SV_", 3) == 0)
+	{
+		_mesa_glsl_warning(ParseState, "unrecognized system value output '%s'",
+			Semantic);
+	}
+
+	*DestVariableType = Type;
+
+	// Create variable
+	glsl_struct_field *StructField = ralloc_array(ParseState, glsl_struct_field, 1);
+
+	memset(StructField, 0, sizeof(glsl_struct_field));
+	StructField[0].type = Type;
+	StructField[0].name = ralloc_strdup(ParseState, "Data");
+
+	const glsl_type* VariableType = glsl_type::get_record_instance(StructField, 1, ralloc_strdup(ParseState, Semantic));
+
+	Variable = new(ParseState)ir_variable(VariableType, ralloc_asprintf(ParseState, "out_%s", Semantic), ir_var_out);
+	Variable->centroid = false;//OutputQualifier.Fields.bCentroid;
+	Variable->interpolation = false;//OutputQualifier.Fields.InterpolationMode;
+	Variable->is_interface_block = true;
+	Variable->is_patch_constant = false;//OutputQualifier.Fields.bIsPatchConstant;
+
+	DeclInstructions->push_tail(Variable);
+	ParseState->symbols->add_variable(Variable);
+
+	ir_rvalue* VariableDeref = new(ParseState)ir_dereference_variable(Variable);
+	VariableDeref = new(ParseState)ir_dereference_record(VariableDeref, ralloc_strdup(ParseState, "Data"));
+
+	return VariableDeref;
+}
+
+
+/**
+* Generate an output semantic.
+* @param Frequency - The shader frequency.
+* @param ParseState - Parse state.
+* @param OutputSemantic - The semantic name to generate.
+* @param OutputQualifier - Qualifiers applied to the semantic.
+* @param OutputVariableDeref - Deref for the argument variable.
+* @param DeclInstructions - IR to which declarations may be added.
+* @param PostCallInstructions - IR to which instructions may be added after the
+*                               entry point returns.
+*/
+void GenShaderOutputForVariable(
+	EHlslShaderFrequency Frequency,
+	_mesa_glsl_parse_state* ParseState,
+	const char* OutputSemantic,
+	ir_dereference* OutputVariableDeref,
+	exec_list* DeclInstructions,
+	exec_list* PostCallInstructions,
+	int SemanticArraySize,
+	int SemanticArrayIndex
+	)
+{
+	const glsl_type* OutputType = OutputVariableDeref->type;
+	if (OutputType->is_record())
+	{
+		check(0);
+	}
+	// TODO clean this up!!
+	else if (OutputType->is_array())// || OutputType->is_outputpatch()))
+	{
+		check(0);
+	}
+	else
+	{
+		if (OutputSemantic)
+		{
+			YYLTYPE loc = {0};
+			ir_rvalue* Src = OutputVariableDeref->clone(ParseState, NULL);
+			const glsl_type* DestVariableType = NULL;
+			ir_rvalue* DestVariableDeref = GenShaderOutputSemantic(Frequency, ParseState, OutputSemantic,
+				OutputType, DeclInstructions, &DestVariableType);
+
+			apply_type_conversion(DestVariableType, Src, PostCallInstructions, ParseState, true, &loc);
+			PostCallInstructions->push_tail(new(ParseState)ir_assignment(DestVariableDeref, Src));
+		}
+		else
+		{
+			_mesa_glsl_error(ParseState, "entry point does not specify a semantic for its return value");
+		}
+	}
+}
+
+
+/**
+* Generate an output semantic.
+* @param Frequency - The shader frequency.
+* @param ParseState - Parse state.
+* @param OutputSemantic - The semantic name to generate.
+* @param OutputQualifier - Qualifiers applied to the semantic.
+* @param OutputType - Value type.
+* @param DeclInstructions - IR to which declarations may be added.
+* @param PreCallInstructions - IR to which isntructions may be added before the
+entry point is called.
+* @param PostCallInstructions - IR to which instructions may be added after the
+*                               entry point returns.
+* @returns the IR variable deref for the semantic.
+*/
+static ir_dereference_variable* GenerateShaderOutput(
+	EHlslShaderFrequency Frequency,
+	_mesa_glsl_parse_state* ParseState,
+	const char* OutputSemantic,
+	const glsl_type* OutputType,
+	exec_list* DeclInstructions,
+	exec_list* PreCallInstructions,
+	exec_list* PostCallInstructions
+	)
+{
+	// Generate a local variable to hold the output.
+	ir_variable* TempVariable = new(ParseState)ir_variable(
+		OutputType,
+		NULL,
+		ir_var_temporary);
+	ir_dereference_variable* TempVariableDeref = new(ParseState)ir_dereference_variable(TempVariable);
+	PreCallInstructions->push_tail(TempVariable);
+	GenShaderOutputForVariable(
+		Frequency,
+		ParseState,
+		OutputSemantic,
+		TempVariableDeref,
+		DeclInstructions,
+		PostCallInstructions,
+		0,
+		0
+		);
+	return TempVariableDeref;
+}
+
+
+void FMetalCodeBackend::PackInputsAndOutputs(exec_list* Instructions, _mesa_glsl_parse_state* ParseState, EHlslShaderFrequency Frequency, exec_list& InputVars)
+{
+	ir_function_signature* EntryPointSig = GetMainFunction(Instructions);
+	check(EntryPointSig);
+
+	exec_list DeclInstructions;
+	exec_list PreCallInstructions;
+	exec_list ArgInstructions;
+	exec_list PostCallInstructions;
+	ParseState->symbols->push_scope();
+
+	// 
+
+	// Set of variables packed into a struct
+	std::set<ir_variable*> VSStageInVariables;
+	std::set<ir_variable*> PSStageInVariables;
+	std::set<ir_variable*> VSOutVariables;
+	std::set<ir_variable*> PSOutVariables;
+
+	// Return var/struct
+	ir_variable* VSOut = nullptr;
+	ir_variable* PSOut = nullptr;
+
+	// Input stage variables
+	ir_variable* VSStageIn = nullptr;
+	std::map<std::string, glsl_struct_field> OriginalVSStageInMembers;
+	ir_variable* PSStageIn = nullptr;
+
+	// Extra arguments needed for input (VertexID, etc)
+	TIRVarList VSInputArguments;
+	TIRVarList PSInputArguments;
+
+	if (Frequency == HSF_VertexShader)
+	{
+		// Vertex Fetch to Vertex connector
+		TArray<glsl_struct_field> VSStageInMembers;
+
+		// Vertex Output connector. Gather position semantic & other outputs into a struct
+		TArray<glsl_struct_field> VSOutMembers;
+
+		foreach_iter(exec_list_iterator, Iter, *Instructions)
+		{
+			ir_instruction* IR = (ir_instruction*)Iter.get();
+			auto* Variable = IR->as_variable();
+			if (Variable)
+			{
+				switch (Variable->mode)
+				{
+				case ir_var_out:
+					{
+						glsl_struct_field Member;
+						Member.type = Variable->type;
+						Member.name = ralloc_strdup(ParseState, Variable->name);
+						Member.semantic = ralloc_strdup(ParseState, Variable->semantic ? Variable->semantic : Variable->name);;
+						VSOutMembers.Add(Member);
+						VSOutVariables.insert(Variable);
+					}
+					break;
+
+				case ir_var_in:
+					if (!ProcessStageInVariables(ParseState, Frequency, Variable, VSStageInMembers, VSStageInVariables, nullptr, VSInputArguments))
+					{
+						return;
+					}
+					break;
+				}
+			}
+		}
+
+		if (VSStageInMembers.Num())
+		{
+			check(Frequency == HSF_VertexShader);
+
+			//@todo-rco: Make me nice...
+			int AttributesUsedMask = 0;
+			for (auto& Member : VSStageInMembers)
+			{
+				int Index = GetAttributeIndex(Member.semantic);
+				if (Index >= 0 && Index < 16)
+				{
+					AttributesUsedMask |= (1 << Index);
+				}
+				InputVars.push_tail(new(ParseState)extern_var(new(ParseState)ir_variable(Member.type, ralloc_strdup(ParseState, Member.name), ir_var_in)));
+			}
+
+			if (bGenerateVSInputDummies)
+			{
+				int Bit = 0;
+				for (int i = 0; i < 16; ++i)
+				{
+					if ((AttributesUsedMask & (1 << i)) == 0)
+					{
+						glsl_struct_field NewMember;
+						NewMember.name = ralloc_asprintf(ParseState, "__dummy%d", i);
+						NewMember.semantic = ralloc_asprintf(ParseState, "ATTRIBUTE%d", i);
+						NewMember.type = glsl_type::get_instance(GLSL_TYPE_FLOAT, 4, 1);
+						VSStageInMembers.Add(NewMember);
+					}
+				}
+			}
+
+			VSStageInMembers.Sort(
+				[](glsl_struct_field const& A, glsl_struct_field const& B)
+				{
+					return GetAttributeIndex(A.semantic) < GetAttributeIndex(B.semantic);
+				});
+
+			// Convert all members to float4
+			if (bExpandVSInputsToFloat4)
+			{
+				for (auto& Member : VSStageInMembers)
+				{
+					OriginalVSStageInMembers[Member.name] = Member;
+					check(Member.type->matrix_columns == 1);
+					Member.type = glsl_type::get_instance(Member.type->base_type, 4, 1);
+				}
+			}
+
+			auto* Type = glsl_type::get_record_instance(&VSStageInMembers[0], (unsigned int)VSStageInMembers.Num(), "FVSStageIn");
+			VSStageIn = new(ParseState)ir_variable(Type, "__VSStageIn", ir_var_in);
+			// Hack: This way we tell we need to convert types from half to float
+			((glsl_type*)Type)->HlslName = "__STAGE_IN__";
+			ParseState->symbols->add_variable(VSStageIn);
+
+			if (!ParseState->AddUserStruct(Type))
+			{
+				YYLTYPE loc = {0};
+				_mesa_glsl_error(&loc, ParseState, "struct '%s' previously defined", Type->name);
+			}
+		}
+
+		if (VSOutMembers.Num())
+		{
+			auto* Type = glsl_type::get_record_instance(&VSOutMembers[0], (unsigned int)VSOutMembers.Num(), "FVSOut");
+			VSOut = new(ParseState)ir_variable(Type, "__VSOut", ir_var_temporary);
+			PostCallInstructions.push_tail(VSOut);
+			ParseState->symbols->add_variable(VSOut);
+
+			if (!ParseState->AddUserStruct(Type))
+			{
+				YYLTYPE loc = {0};
+				_mesa_glsl_error(&loc, ParseState, "struct '%s' previously defined", Type->name);
+			}
+		}
+	}
+	else if (Frequency == HSF_PixelShader)
+	{
+		// Vertex to Pixel connector
+		TArray<glsl_struct_field> PSStageInMembers;
+
+		// Pixel Output connector. Gather color & depth outputs into a struct
+		TArray<glsl_struct_field> PSOutMembers;
+
+		// Gather all inputs and generate the StageIn VS->PS connector
+		foreach_iter(exec_list_iterator, Iter, *Instructions)
+		{
+			ir_instruction* IR = (ir_instruction*)Iter.get();
+			auto* Variable = IR->as_variable();
+			if (Variable)
+			{
+				switch (Variable->mode)
+				{
+				case ir_var_out:
+					{
+						glsl_struct_field Member;
+						Member.type = Variable->type;
+						Member.name = ralloc_strdup(ParseState, Variable->name);
+						Member.semantic = ralloc_strdup(ParseState, Variable->semantic ? Variable->semantic : Variable->name);
+						PSOutMembers.Add(Member);
+						PSOutVariables.insert(Variable);
+					}
+					break;
+
+				case ir_var_in:
+					if (!ProcessStageInVariables(ParseState, Frequency, Variable, PSStageInMembers, PSStageInVariables, nullptr, PSInputArguments))
+					{
+						return;
+					}
+					break;
+
+				default:
+					break;
+				}
+			}
+		}
+
+		if (PSStageInMembers.Num())
+		{
+			auto* Type = glsl_type::get_record_instance(&PSStageInMembers[0], (unsigned int)PSStageInMembers.Num(), "FPSStageIn");
+			// Hack: This way we tell we need to convert types from half to float
+			((glsl_type*)Type)->HlslName = "__STAGE_IN__";
+			PSStageIn = new(ParseState)ir_variable(Type, "__PSStageIn", ir_var_in);
+			ParseState->symbols->add_variable(PSStageIn);
+
+			if (!ParseState->AddUserStruct(Type))
+			{
+				YYLTYPE loc = {0};
+				_mesa_glsl_error(&loc, ParseState, "struct '%s' previously defined", Type->name);
+			}
+		}
+
+		if (PSOutMembers.Num())
+		{
+			auto* Type = glsl_type::get_record_instance(&PSOutMembers[0], (unsigned int)PSOutMembers.Num(), "FPSOut");
+			PSOut = new(ParseState)ir_variable(Type, "__PSOut", ir_var_temporary);
+			PostCallInstructions.push_tail(PSOut);
+			ParseState->symbols->add_variable(PSOut);
+
+			if (!ParseState->AddUserStruct(Type))
+			{
+				YYLTYPE loc = {0};
+				_mesa_glsl_error(&loc, ParseState, "struct '%s' previously defined", Type->name);
+			}
+		}
+	}
+
+	TIRVarList VarsToMoveToBody;
+	foreach_iter(exec_list_iterator, Iter, *Instructions)
+	{
+		ir_instruction* IR = (ir_instruction*)Iter.get();
+		auto* Variable = IR->as_variable();
+		if (Variable)
+		{
+			ir_dereference_variable* ArgVarDeref = NULL;
+			switch (Variable->mode)
+			{
+			case ir_var_in:
+				if (PSStageInVariables.find(Variable) != PSStageInVariables.end())
+				{
+					ir_dereference* DeRefMember = new(ParseState)ir_dereference_record(PSStageIn, Variable->name);
+					auto* Assign = new(ParseState)ir_assignment(new(ParseState)ir_dereference_variable(Variable), DeRefMember);
+					PreCallInstructions.push_tail(Assign);
+					VarsToMoveToBody.push_back(Variable);
+				}
+				else if (VSStageInVariables.find(Variable) != VSStageInVariables.end())
+				{
+					ir_rvalue* DeRefMember = new(ParseState)ir_dereference_record(VSStageIn, Variable->name);
+					unsigned int Mask = 0;
+					if (bExpandVSInputsToFloat4)
+					{
+						Mask = (1 << 4) - 1;
+						auto FoundMember = OriginalVSStageInMembers.find(Variable->name);
+						check(FoundMember != OriginalVSStageInMembers.end());
+						if (FoundMember->second.type)
+						{
+							check(FoundMember->second.type->vector_elements != 0);
+							Mask = (1 << FoundMember->second.type->vector_elements) - 1;
+							if (Mask != 15)
+							{
+								DeRefMember = new(ParseState)ir_swizzle(DeRefMember, 0, 1, 2, 3, FoundMember->second.type->vector_elements);
+							}
+						}
+					}
+					else
+					{
+						Mask = (1 << Variable->type->vector_elements) - 1;
+					}
+					auto* Assign = new(ParseState)ir_assignment(new(ParseState)ir_dereference_variable(Variable), DeRefMember, nullptr, Mask);
+					PreCallInstructions.push_tail(Assign);
+					VarsToMoveToBody.push_back(Variable);
+				}
+				else
+				{
+					// At this point this should be a built-in system value
+					check(Variable->semantic);
+					ArgVarDeref = GenerateShaderInput(
+						Frequency,
+						ParseState,
+						Variable->semantic,
+						Variable->type,
+						&DeclInstructions,
+						&PreCallInstructions
+						);
+				}
+				break;
+
+			case ir_var_out:
+				if (VSOutVariables.find(Variable) != VSOutVariables.end())
+				{
+					VarsToMoveToBody.push_back(Variable);
+					ir_dereference* DeRefMember = new(ParseState)ir_dereference_record(VSOut, Variable->name);
+					auto* Assign = new(ParseState)ir_assignment(DeRefMember, new(ParseState)ir_dereference_variable(Variable));
+					PostCallInstructions.push_tail(Assign);
+				}
+				else if (PSOutVariables.find(Variable) != PSOutVariables.end())
+				{
+					VarsToMoveToBody.push_back(Variable);
+					ir_dereference* DeRefMember = new(ParseState)ir_dereference_record(PSOut, Variable->name);
+					auto* Assign = new(ParseState)ir_assignment(DeRefMember, new(ParseState)ir_dereference_variable(Variable));
+					PostCallInstructions.push_tail(Assign);
+				}
+				else
+				{
+					ArgVarDeref = GenerateShaderOutput(
+						Frequency,
+						ParseState,
+						Variable->semantic,
+						Variable->type,
+						&DeclInstructions,
+						&PreCallInstructions,
+						&PostCallInstructions
+						);
+				}
+				break;
+
+			default:
+				break;
+				/*
+				_mesa_glsl_error(
+				ParseState,
+				"entry point parameter '%s' must be an input or output",
+				Variable->name
+				);
+				*/
+			}
+			if (ArgVarDeref)
+			{
+				ArgInstructions.push_tail(ArgVarDeref);
+			}
+		}
+		/*
+		else
+		{
+		_mesa_glsl_error(ParseState, "entry point parameter "
+		"'%s' does not specify a semantic", Variable->name);
+		}
+		*/
+	}
+
+	// The function's return value should have an output semantic if it's not void.
+	ir_dereference_variable* EntryPointReturn = NULL;
+	if (EntryPointSig->return_type->is_void() == false)
+	{
+		if (Frequency == HSF_PixelShader)
+		{
+			check(!PSOut);
+			PSOut = new(ParseState)ir_variable(EntryPointSig->return_type, "__PSOut", ir_var_temporary);
+			PreCallInstructions.push_tail(PSOut);
+		}
+		else if (Frequency == HSF_VertexShader)
+		{
+			check(!VSOut);
+			VSOut = new(ParseState)ir_variable(EntryPointSig->return_type, "__VSOut", ir_var_temporary);
+			PreCallInstructions.push_tail(VSOut);
+		}
+		else
+		{
+			check(0);
+		}
+	}
+
+	ParseState->symbols->pop_scope();
+
+	// Build the void main() function for GLSL.
+	auto* ReturnType = glsl_type::void_type;
+	if (VSOut)
+	{
+		ReturnType = VSOut->type;
+		check(!EntryPointReturn);
+		EntryPointReturn = new(ParseState)ir_dereference_variable(VSOut);
+		PostCallInstructions.push_tail(new(ParseState)ir_return(new(ParseState)ir_dereference_variable(VSOut)));
+		EntryPointSig->return_type = ReturnType;
+	}
+	else if (PSOut)
+	{
+		ReturnType = PSOut->type;
+		check(!EntryPointReturn);
+		EntryPointReturn = new(ParseState)ir_dereference_variable(PSOut);
+		PostCallInstructions.push_tail(new(ParseState)ir_return(new(ParseState)ir_dereference_variable(PSOut)));
+		EntryPointSig->return_type = ReturnType;
+	}
+
+	for (auto* Var : VarsToMoveToBody)
+	{
+		Var->remove();
+		if (Var->mode == ir_var_in || Var->mode == ir_var_out)
+		{
+			Var->mode = ir_var_temporary;
+		}
+		DeclInstructions.push_head(Var);
+	}
+
+	DeclInstructions.append_list(&PreCallInstructions);
+	DeclInstructions.append_list(&EntryPointSig->body);
+	DeclInstructions.append_list(&PostCallInstructions);
+
+	EntryPointSig->body.append_list(&DeclInstructions);
+#if 0
+	Instructions->push_tail(MainFunction);
+#endif // 0
+
+	// Now that we have a proper main(), move global setup to main().
+	if (VSStageIn)
+	{
+		EntryPointSig->parameters.push_tail(VSStageIn);
+	}
+	else if (PSStageIn)
+	{
+		EntryPointSig->parameters.push_tail(PSStageIn);
+	}
+	/*
+	MoveGlobalInstructionsToMain(Instructions);
+	*/
+}
+
+struct FConvertHalfToFloatUniformAndSamples : public ir_rvalue_visitor
+{
+	struct FPair
+	{
+		ir_rvalue** RValuePtr;
+		ir_instruction* InsertPoint;
+	};
+	typedef std::map<ir_rvalue*, std::list<FPair>> TReplacedVarMap;
+	TReplacedVarMap ReplacedVars;
+
+	std::list<TReplacedVarMap> PendingReplacements;
+	TIRVarSet ReferencedUniforms;
+
+	_mesa_glsl_parse_state* State;
+
+	bool bIsMaster;
+	bool bConvertUniforms;
+	bool bConvertSamples;
+
+	FConvertHalfToFloatUniformAndSamples(_mesa_glsl_parse_state* InState, bool bInConvertUniforms, bool bInConvertSamples) :
+		State(InState),
+		bIsMaster(true),
+		bConvertUniforms(bInConvertUniforms),
+		bConvertSamples(bInConvertSamples)
+	{
+	}
+
+	void DoConvertOneMap(TReplacedVarMap& Map)
+	{
+		for (auto& TopIter : Map)
+		{
+			auto* RValue = TopIter.first;
+			auto& List = TopIter.second;
+
+			// Coerce this var into float
+			auto* OriginalVar = RValue->variable_referenced();
+			const auto* OriginalVarType = OriginalVar->type;
+			const auto* PromotedVarType = PromoteHalfToFloatType(State, OriginalVarType);
+			OriginalVar->type = PromotedVarType;
+
+			// Temp var and assignment
+			auto* NewVar = new(State)ir_variable(RValue->type, nullptr, ir_var_temporary);
+			RValue->type = PromoteHalfToFloatType(State, RValue->type);
+			exec_list NewAssignments;
+			CreateNewAssignmentsFloat2Half(State, NewAssignments, NewVar, RValue);
+			auto* BaseIR = List.front().InsertPoint;
+
+			// Store new instructions so we add a nice block in the asm
+			BaseIR->insert_before(NewVar);
+			BaseIR->insert_before(&NewAssignments);
+
+			for (auto& Iter : List)
+			{
+				*(Iter.RValuePtr) = new(State)ir_dereference_variable(NewVar);
+			}
+		}
+
+		// Go through all remaining parameters
+		for (auto& Var : ReferencedUniforms)
+		{
+			Var->type = PromoteHalfToFloatType(State, Var->type);
+		}
+	}
+
+	void DoConvert(exec_list* ir)
+	{
+		run(ir);
+		DoConvertOneMap(ReplacedVars);
+
+		if (bIsMaster)
+		{
+			for (auto& Map : PendingReplacements)
+			{
+				DoConvertOneMap(Map);
+			}
+		}
+	}
+
+	void ConvertBlock(exec_list* instructions)
+	{
+		FConvertHalfToFloatUniformAndSamples Visitor(State, bConvertUniforms, bConvertSamples);
+		Visitor.bIsMaster = false;
+		Visitor.run(instructions);
+		PendingReplacements.push_back(TReplacedVarMap());
+		PendingReplacements.back().swap(Visitor.ReplacedVars);
+		for (auto& Var : Visitor.ReferencedUniforms)
+		{
+			ReferencedUniforms.insert(Var);
+		}
+	}
+
+	virtual ir_visitor_status visit_enter(ir_if* ir) override
+	{
+		ir->condition->accept(this);
+		handle_rvalue(&ir->condition);
+
+		ConvertBlock(&ir->then_instructions);
+		ConvertBlock(&ir->else_instructions);
+
+		/* already descended into the children. */
+		return visit_continue_with_parent;
+	}
+
+	ir_visitor_status visit_enter(ir_loop* ir)
+	{
+		ConvertBlock(&ir->body_instructions);
+
+		/* already descended into the children. */
+		return visit_continue_with_parent;
+	}
+
+	virtual void handle_rvalue(ir_rvalue** RValuePtr) override
+	{
+		if (!RValuePtr || !*RValuePtr)
+		{
+			return;
+		}
+
+		auto* RValue = *RValuePtr;
+		if (bConvertSamples && RValue->as_texture())
+		{
+			auto* Texture = RValue->as_texture();
+			if (Texture->coordinate && Texture->coordinate->type->base_type == GLSL_TYPE_HALF)
+			{
+				// Promote to float
+				Texture->coordinate = new(State)ir_expression(ir_unop_h2f, Texture->coordinate);
+			}
+			else if (Texture->coordinate && Texture->coordinate->type->base_type == GLSL_TYPE_INT)
+			{
+				// convert int to uint
+				Texture->coordinate = new(State)ir_expression(ir_unop_i2u, Texture->coordinate);
+			}
+		}
+		// Skip swizzles, textures, etc
+		else if (bConvertUniforms && RValue->as_dereference())
+		{
+			auto* Var = RValue->variable_referenced();
+			if (Var && Var->mode == ir_var_uniform)
+			{
+				ReferencedUniforms.insert(Var);
+				if (RValue->type->base_type == GLSL_TYPE_HALF)
+				{
+					// Save this RValue and prep for later change
+					FPair Pair = {RValuePtr, base_ir};
+					for (auto& Iter : ReplacedVars)
+					{
+						auto* TestRValue = Iter.first;
+						if (RValue->ir_type == TestRValue->ir_type && AreEquivalent(RValue, TestRValue))
+						{
+							Iter.second.push_back(Pair);
+							return;
+						}
+					}
+
+					ReplacedVars[RValue].push_back(Pair);
+				}
+			}
+		}
+	}
+};
+
+
+void FMetalCodeBackend::ConvertHalfToFloatUniformsAndSamples(exec_list* ir, _mesa_glsl_parse_state* State, bool bConvertUniforms, bool bConvertSamples)
+{
+	if (bConvertUniforms || bConvertSamples)
+	{
+		FConvertHalfToFloatUniformAndSamples ConvertHalfToFloatUniforms(State, bConvertUniforms, bConvertSamples);
+		ConvertHalfToFloatUniforms.DoConvert(ir);
+	}
+}
+
+struct FMetalBreakPrecisionChangesVisitor : public ir_rvalue_visitor
+{
+	_mesa_glsl_parse_state* State;
+
+	std::map<ir_rvalue*, ir_variable*> ReplacedVars;
+
+	FMetalBreakPrecisionChangesVisitor(_mesa_glsl_parse_state* InState) : State(InState) {}
+
+	void ConvertBlock(exec_list* instructions)
+	{
+		FMetalBreakPrecisionChangesVisitor Visitor(State);
+		Visitor.run(instructions);
+	}
+
+	virtual ir_visitor_status visit_enter(ir_if* ir) override
+	{
+		ir->condition->accept(this);
+		handle_rvalue(&ir->condition);
+
+		ConvertBlock(&ir->then_instructions);
+		ConvertBlock(&ir->else_instructions);
+
+		/* already descended into the children. */
+		return visit_continue_with_parent;
+	}
+
+	ir_visitor_status visit_enter(ir_loop* ir)
+	{
+		ConvertBlock(&ir->body_instructions);
+
+		/* already descended into the children. */
+		return visit_continue_with_parent;
+	}
+
+	virtual void handle_rvalue(ir_rvalue** RValuePtr) override
+	{
+		if (!RValuePtr || !*RValuePtr)
+		{
+			return;
+		}
+		bool bGenerateNewVar = false;
+		auto* RValue = *RValuePtr;
+		auto* Expression = RValue->as_expression();
+		auto* Constant = RValue->as_constant();
+		if (Expression)
+		{
+			switch (Expression->operation)
+			{
+			case ir_unop_h2f:
+			case ir_unop_f2h:
+				if (!Expression->operands[0]->as_texture())
+				{
+					bGenerateNewVar = true;
+				}
+				break;
+			}
+		}
+		else if (Constant)
+		{
+			if (Constant->type->base_type == GLSL_TYPE_HALF)
+			{
+				bGenerateNewVar = true;
+			}
+		}
+
+		if (bGenerateNewVar)
+		{
+			for (auto& Iter : ReplacedVars)
+			{
+				auto* TestRValue = Iter.first;
+				if (AreEquivalent(TestRValue, RValue))
+				{
+					*RValuePtr = new(State)ir_dereference_variable(Iter.second);
+					return;
+				}
+			}
+
+			auto* NewVar = new(State)ir_variable(RValue->type, nullptr, ir_var_temporary);
+			auto* NewAssignment = new(State)ir_assignment(new(State)ir_dereference_variable(NewVar), RValue);
+			*RValuePtr = new(State)ir_dereference_variable(NewVar);
+			ReplacedVars[RValue] = NewVar;
+			base_ir->insert_before(NewVar);
+			base_ir->insert_before(NewAssignment);
+		}
+	}
+};
+
+
+void FMetalCodeBackend::BreakPrecisionChangesVisitor(exec_list* ir, _mesa_glsl_parse_state* State)
+{
+	FMetalBreakPrecisionChangesVisitor Visitor(State);
+	Visitor.run(ir);
+}
+
+struct FDeReferencePackedVarsVisitor : public ir_rvalue_visitor
+{
+	_mesa_glsl_parse_state* State;
+	int ExpressionDepth;
+
+	FDeReferencePackedVarsVisitor(_mesa_glsl_parse_state* InState)
+		: State(InState)
+		, ExpressionDepth(0)
+	{
+	}
+
+	virtual ir_visitor_status visit_enter(ir_expression* ir) override
+	{
+		++ExpressionDepth;
+		return ir_rvalue_visitor::visit_enter(ir);
+	}
+
+	virtual ir_visitor_status visit_leave(ir_expression* ir) override
+	{
+		for (uint32 i = 0; i < ir->get_num_operands(); ++i)
+		{
+			auto* Operand = ir->operands[i];
+			auto* DerefRecord = Operand->as_dereference_record();
+			if (DerefRecord)
+			{
+				handle_rvalue(&ir->operands[i]);
+			}
+		}
+
+		--ExpressionDepth;
+		return ir_rvalue_visitor::visit_leave(ir);
+	}
+
+	virtual void handle_rvalue(ir_rvalue** RValuePtr) override
+	{
+		if (!RValuePtr || !*RValuePtr)
+		{
+			return;
+		}
+
+		auto* DeRefRecord = (*RValuePtr)->as_dereference_record();
+		auto* Swizzle = (*RValuePtr)->as_swizzle();
+		auto* SwizzleValDeRefRecord = Swizzle ? Swizzle->val->as_dereference_record() : nullptr;
+		if (SwizzleValDeRefRecord)
+		{
+			auto* StructVar = Swizzle->variable_referenced();
+			if (StructVar->type->HlslName && !strcmp(StructVar->type->HlslName, "__PACKED__"))
+			{
+				if (SwizzleValDeRefRecord->type->vector_elements > 1 && SwizzleValDeRefRecord->type->vector_elements < 4)
+				{
+					auto* Var = GetVar(SwizzleValDeRefRecord);
+					Swizzle->val = new(State)ir_dereference_variable(Var);
+				}
+			}
+		}
+		else if (DeRefRecord && ExpressionDepth > 0)
+		{
+			auto* StructVar = DeRefRecord->variable_referenced();
+			if (StructVar->type->HlslName && !strcmp(StructVar->type->HlslName, "__PACKED__"))
+			{
+				if (DeRefRecord->type->vector_elements > 1 && DeRefRecord->type->vector_elements < 4)
+				{
+					auto* Var = GetVar(DeRefRecord);
+					*RValuePtr = new(State)ir_dereference_variable(Var);
+				}
+			}
+		}
+	}
+
+	std::map<ir_dereference_record*, ir_variable*> Replaced;
+
+	ir_variable* GetVar(ir_dereference_record* ir)
+	{
+		ir_variable* Var = nullptr;
+		for (auto& Pair : Replaced)
+		{
+			if (Pair.first->IsEquivalent(ir))
+			{
+				Var = Pair.second;
+				break;
+			}
+		}
+
+		if (!Var)
+		{
+			Var = new(State)ir_variable(ir->type, nullptr, ir_var_temporary);
+			Replaced[ir] = Var;
+		}
+		return Var;
+	}
+};
+
+void FMetalCodeBackend::RemovePackedVarReferences(exec_list* ir, _mesa_glsl_parse_state* State)
+{
+	FDeReferencePackedVarsVisitor Visitor(State);
+	Visitor.run(ir);
+
+	auto* Main = Visitor.Replaced.empty() ? nullptr : GetMainFunction(ir);
+	for (auto& OuterPair : Visitor.Replaced)
+	{
+		auto* NewVar = OuterPair.second;
+		auto* DeRefRecord = OuterPair.first;
+		auto* NewAssignment = new(State)ir_assignment(new(State)ir_dereference_variable(NewVar), DeRefRecord);
+		Main->body.push_head(NewAssignment);
+		Main->body.push_head(NewVar);
+	}
+}
diff --git a/Engine/Source/Developer/iOS/MetalShaderFormat/Private/MetalUtils.h b/Engine/Source/Developer/iOS/MetalShaderFormat/Private/MetalUtils.h
new file mode 100644
--- /dev/null
+++ b/Engine/Source/Developer/iOS/MetalShaderFormat/Private/MetalUtils.h
@@ -0,0 +1,109 @@
+// Copyright 1998-2014 Epic Games, Inc. All Rights Reserved.
+
+#pragma once
+
+#include "hlslcc.h"
+
+inline std::string FixVecPrefix(std::string Type)
+{
+	if (!strncmp("vec", Type.c_str(), 3))
+	{
+		std::string Num = Type.substr(3);
+		Type = "float";
+		Type += Num;
+	}
+	else if (!strncmp("ivec", Type.c_str(), 4))
+	{
+		std::string Num = Type.substr(4);
+		Type = "int";
+		Type += Num;
+	}
+	else if (!strncmp("uvec", Type.c_str(), 4))
+	{
+		std::string Num = Type.substr(4);
+		Type = "uint";
+		Type += Num;
+	}
+
+	return Type;
+}
+
+/** Track external variables. */
+struct extern_var : public exec_node
+{
+	ir_variable* var;
+	explicit extern_var(ir_variable* in_var) : var(in_var) {}
+};
+
+
+struct FBuffers
+{
+	TArray<ir_instruction*> Buffers;
+
+	int GetIndex(ir_variable* Var)
+	{
+		for (int i = 0, n = Buffers.Num(); i < n; ++i)
+		{
+			if (Buffers[i] == Var)
+			{
+				return i;
+			}
+		}
+
+		return -1;
+	}
+
+	void SortBuffers()
+	{
+		TArray<ir_instruction*> AllBuffers;
+		AllBuffers.AddZeroed(Buffers.Num());
+		TIRVarList CBuffers;
+		// Put packed UB's into their location (h=0, m=1, etc); leave holes if not using a packed define
+		// and group the regular CBuffers in another list
+		for (int i = 0, n = Buffers.Num(); i < n; ++i)
+		{
+			auto* Var = Buffers[i]->as_variable();
+			check(Var);
+			if (Var->semantic && strlen(Var->semantic) == 1)
+			{
+				int Index = ConvertArrayTypeToIndex((EArrayType)Var->semantic[0]);
+				if (AllBuffers.Num() <= Index)
+				{
+					int32 Count = Index + 1 - AllBuffers.Num();
+					AllBuffers.AddZeroed(Count);
+					//AllBuffers.Resize(Index + 1, nullptr);
+				}
+				AllBuffers[Index] = Var;
+			}
+			else
+			{
+				CBuffers.push_back(Var);
+			}
+		}
+
+		// Fill the holes in the packed array list with real UB's
+		for (int i = 0; i < AllBuffers.Num() && !CBuffers.empty(); ++i)
+		{
+			if (!AllBuffers[i])
+			{
+				AllBuffers[i] = CBuffers.front();
+				CBuffers.erase(CBuffers.begin());
+			}
+		}
+
+		Buffers = AllBuffers;
+	}
+};
+
+const glsl_type* PromoteHalfToFloatType(_mesa_glsl_parse_state* state, const glsl_type* type);
+
+namespace MetalUtils
+{
+	ir_dereference_variable* GenerateInput(EHlslShaderFrequency Frequency, _mesa_glsl_parse_state* ParseState, const char* InputSemantic,
+		const glsl_type* InputType, exec_list* DeclInstructions, exec_list* PreCallInstructions);
+
+	ir_dereference_variable* GenerateOutput(EHlslShaderFrequency Frequency, _mesa_glsl_parse_state* ParseState, const char* OutputSemantic,
+		const glsl_type* OutputType, exec_list* DeclInstructions, exec_list* PreCallInstructions, exec_list* PostCallInstructions);
+}
+
+const int MAX_SIMULTANEOUS_RENDER_TARGETS = 8;
