Fix Mac CIS.

- Back out revision 6 from //depot/UE4/Engine/Source/Developer/iOS/MetalShaderFormat/Private/MetalBackend.cpp
- itoa() is non-standard.

diff --git a/Engine/Source/Developer/iOS/MetalShaderFormat/Private/MetalBackend.cpp b/Engine/Source/Developer/iOS/MetalShaderFormat/Private/MetalBackend.cpp
--- a/Engine/Source/Developer/iOS/MetalShaderFormat/Private/MetalBackend.cpp
+++ b/Engine/Source/Developer/iOS/MetalShaderFormat/Private/MetalBackend.cpp
@@ -18,7 +18,6 @@
 #if !PLATFORM_WINDOWS
 #define _stricmp stricmp
 #define _strnicmp strnicmp
-#define _itoa itoa
 #endif
 
 /**
@@ -2659,8 +2658,8 @@ void FMetalLanguageSpec::SetupLanguageIntrinsics(_mesa_glsl_parse_state* State,
 		// MRTs; first make intrinsics for each MRT, then a non-intrinsic version to use that (helps when converting to Metal)
 		for (int i = 0; i < MAX_SIMULTANEOUS_RENDER_TARGETS; ++i)
 		{
-			char FunctionName[32] = FRAMEBUFFER_FETCH_MRT;
-			_itoa(i, FunctionName + strlen(FRAMEBUFFER_FETCH_MRT), 10);
+			char FunctionName[32];
+			sprintf(FunctionName, "%s%d", FRAMEBUFFER_FETCH_MRT, i);
 			make_intrinsic_genType(ir, State, FunctionName, ir_invalid_opcode, IR_INTRINSIC_HALF, 0, 4, 4);
 		}
 
@@ -2678,8 +2677,8 @@ void FMetalLanguageSpec::SetupLanguageIntrinsics(_mesa_glsl_parse_state* State,
 			//	if (Arg0 == i) FRAMEBUFFER_FETCH_MRT#i();
 			auto* Condition = new(State) ir_expression(ir_binop_equal, new(State) ir_dereference_variable((ir_variable*)Sig->parameters.get_head()), new(State) ir_constant(i));
 			auto* If = new(State) ir_if(Condition);
-			char FunctionName[32] = FRAMEBUFFER_FETCH_MRT;
-			_itoa(i, FunctionName + strlen(FRAMEBUFFER_FETCH_MRT), 10);
+			char FunctionName[32];
+			sprintf(FunctionName, "%s%d", FRAMEBUFFER_FETCH_MRT, i);
 			auto* IntrinsicSig = FCodeBackend::FindEntryPointFunction(ir, State, FunctionName);
 			auto* ReturnValue = new(State) ir_variable(ReturnType, nullptr, ir_var_temporary);
 			exec_list Empty;
