fix for sprintf wanring
#ue4

diff --git a/Engine/Source/Developer/iOS/MetalShaderFormat/Private/MetalBackend.cpp b/Engine/Source/Developer/iOS/MetalShaderFormat/Private/MetalBackend.cpp
--- a/Engine/Source/Developer/iOS/MetalShaderFormat/Private/MetalBackend.cpp
+++ b/Engine/Source/Developer/iOS/MetalShaderFormat/Private/MetalBackend.cpp
@@ -2659,7 +2659,7 @@ void FMetalLanguageSpec::SetupLanguageIntrinsics(_mesa_glsl_parse_state* State,
 		for (int i = 0; i < MAX_SIMULTANEOUS_RENDER_TARGETS; ++i)
 		{
 			char FunctionName[32];
-			sprintf(FunctionName, "%s%d", FRAMEBUFFER_FETCH_MRT, i);
+			sprintf_s(FunctionName, 32, "%s%d", FRAMEBUFFER_FETCH_MRT, i);
 			make_intrinsic_genType(ir, State, FunctionName, ir_invalid_opcode, IR_INTRINSIC_HALF, 0, 4, 4);
 		}
 
@@ -2678,7 +2678,7 @@ void FMetalLanguageSpec::SetupLanguageIntrinsics(_mesa_glsl_parse_state* State,
 			auto* Condition = new(State) ir_expression(ir_binop_equal, new(State) ir_dereference_variable((ir_variable*)Sig->parameters.get_head()), new(State) ir_constant(i));
 			auto* If = new(State) ir_if(Condition);
 			char FunctionName[32];
-			sprintf(FunctionName, "%s%d", FRAMEBUFFER_FETCH_MRT, i);
+			sprintf_s(FunctionName, 32, "%s%d", FRAMEBUFFER_FETCH_MRT, i);
 			auto* IntrinsicSig = FCodeBackend::FindEntryPointFunction(ir, State, FunctionName);
 			auto* ReturnValue = new(State) ir_variable(ReturnType, nullptr, ir_var_temporary);
 			exec_list Empty;
