UE4 - MetalBackend - Fix samplerBuffer support
CCT - Fix missing default entry point

diff --git a/Engine/Source/Developer/iOS/MetalShaderFormat/Private/MetalBackend.cpp b/Engine/Source/Developer/iOS/MetalShaderFormat/Private/MetalBackend.cpp
--- a/Engine/Source/Developer/iOS/MetalShaderFormat/Private/MetalBackend.cpp
+++ b/Engine/Source/Developer/iOS/MetalShaderFormat/Private/MetalBackend.cpp
@@ -510,7 +510,13 @@ protected:
 		{
 			ralloc_asprintf_append(buffer, "%s", t->name);
 		}
-		else 
+		else if (t->base_type == GLSL_TYPE_SAMPLER && t->sampler_buffer) 
+		{
+			// Typed buffer read
+			check(t->inner_type);
+			print_base_type(t->inner_type);
+		}
+		else
 		{
 			check(t->HlslName);
 			if (/*!bIsFunctionSig &&*/ bUsePacked && t->is_vector() && t->vector_elements < 4)
@@ -639,42 +645,62 @@ protected:
 				{
 					auto* PtrType = var->type->is_array() ? var->type->element_type() : var->type;
 					check(!PtrType->is_array());
-
 					if (var->type->is_sampler())
 					{
-						auto* Entry = ParseState->FindPackedSamplerEntry(var->name);
-						check(Entry);
-						//@todo-rco: SamplerStates
-						auto SamplerStateFound = ParseState->TextureToSamplerMap.find(Entry->Name.c_str());
-						if (SamplerStateFound != ParseState->TextureToSamplerMap.end())
+						if (var->type->sampler_buffer)
 						{
-							auto& SamplerStates = SamplerStateFound->second;
-							if (SamplerStates.size() == 1)
-							{
-								ralloc_asprintf_append(
-									buffer,
-									"sampler s%d [[ sampler(%d) ]], ", Entry->offset, Entry->offset);
-							}
-							else
+							// Buffer
+							int BufferIndex = Buffers.GetIndex(var);
+							check(BufferIndex >= 0);
+							ralloc_asprintf_append(
+								buffer,
+								"constant "
+								);
+							print_type_pre(PtrType);
+							ralloc_asprintf_append(buffer, " *%s", unique_name(var));
+							print_type_post(PtrType);
+							ralloc_asprintf_append(
+								buffer,
+								" [[ buffer(%d) ]]", BufferIndex
+								);
+						}
+						else
+						{
+							// Regular textures
+							auto* Entry = ParseState->FindPackedSamplerEntry(var->name);
+							check(Entry);
+							//@todo-rco: SamplerStates
+							auto SamplerStateFound = ParseState->TextureToSamplerMap.find(Entry->Name.c_str());
+							if (SamplerStateFound != ParseState->TextureToSamplerMap.end())
 							{
-								check(SamplerStates.empty());
+								auto& SamplerStates = SamplerStateFound->second;
+								if (SamplerStates.size() == 1)
+								{
+									ralloc_asprintf_append(
+										buffer,
+										"sampler s%d [[ sampler(%d) ]], ", Entry->offset, Entry->offset);
+								}
+								else
+								{
+									check(SamplerStates.empty());
+								}
 							}
-						}
 
-						print_type_pre(PtrType);
-						ralloc_asprintf_append(
-							buffer,
-							"<%s> %s", (PtrType->inner_type && PtrType->inner_type->base_type == GLSL_TYPE_HALF) ? "half" : "float", unique_name(var));
-						print_type_post(PtrType);
-						ralloc_asprintf_append(
-							buffer,
-							" [[ texture(%d) ]]", Entry->offset
-							);
+							print_type_pre(PtrType);
+							ralloc_asprintf_append(
+								buffer,
+								"<%s> %s", (PtrType->inner_type && PtrType->inner_type->base_type == GLSL_TYPE_HALF) ? "half" : "float", unique_name(var));
+							print_type_post(PtrType);
+							ralloc_asprintf_append(
+								buffer,
+								" [[ texture(%d) ]]", Entry->offset
+								);
+						}
 					}
 					else
 					{
 						int BufferIndex = Buffers.GetIndex(var);
-						bool bNeedsPointer = var->semantic && (strlen(var->semantic) == 1);
+						bool bNeedsPointer = (var->semantic && (strlen(var->semantic) == 1));
 						check(BufferIndex >= 0);
 						ralloc_asprintf_append(
 							buffer,
diff --git a/Engine/Source/Programs/CrossCompilerTool/Private/CrossCompilerUtils.cpp b/Engine/Source/Programs/CrossCompilerTool/Private/CrossCompilerUtils.cpp
--- a/Engine/Source/Programs/CrossCompilerTool/Private/CrossCompilerUtils.cpp
+++ b/Engine/Source/Programs/CrossCompilerTool/Private/CrossCompilerUtils.cpp
@@ -249,7 +249,7 @@ namespace CCT
 		if (Entry == TEXT(""))
 		{
 			UE_LOG(LogCrossCompilerTool, Warning, TEXT("No entry point given, assuming Main (-entry=Main)!"));
-			Target = HCT_FeatureLevelSM5;
+			Entry = TEXT("Main");
 		}
 
 		return true;
