- Fix some Metal math function names

diff --git a/Engine/Source/Developer/iOS/MetalShaderFormat/Private/MetalBackend.cpp b/Engine/Source/Developer/iOS/MetalShaderFormat/Private/MetalBackend.cpp
--- a/Engine/Source/Developer/iOS/MetalShaderFormat/Private/MetalBackend.cpp
+++ b/Engine/Source/Developer/iOS/MetalShaderFormat/Private/MetalBackend.cpp
@@ -130,7 +130,7 @@ static const char * const GLSLExpressionTable[ir_opcode_count][4] =
 	*
 	* Does not take integer types.
 	*/
-	{ "modf(", ",", ")", "%" }, // ir_binop_mod,
+	{ "fmod(", ",", ")", "%" }, // ir_binop_mod,
 	{ "modf(", ",", ")", "" }, // ir_binop_modf,
 
 	{ "step(", ",", ")", "" }, // ir_binop_step,
@@ -169,15 +169,15 @@ static const char * const GLSLExpressionTable[ir_opcode_count][4] =
 	{ "(", "|", ")", "" }, // ir_binop_bit_or,
 	/*@}*/
 
-	{ "bvec%d(uvec%d(", ")*uvec%d(", "))", "&&" }, // ir_binop_logic_and,
-	{ "bvec%d(abs(ivec%d(", ")+ivec%d(", ")))", "^^" }, // ir_binop_logic_xor,
-	{ "bvec%d(uvec%d(", ")+uvec%d(", "))", "||" }, // ir_binop_logic_or,
+	{ "bool%d(uint%d(", ")*uint%d(", "))", "&&" }, // ir_binop_logic_and,
+	{ "bool%d(abs(int%d(", ")+int%d(", ")))", "^^" }, // ir_binop_logic_xor,
+	{ "bool%d(uint%d(", ")+uint%d(", "))", "||" }, // ir_binop_logic_or,
 
 	{ "dot(", ",", ")", "" }, // ir_binop_dot,
 	{ "cross(", ",", ")", "" }, // ir_binop_cross,
 	{ "fmin(", ",", ")", "" }, // ir_binop_min,
 	{ "fmax(", ",", ")", "" }, // ir_binop_max,
-	{ "atan(", ",", ")", "" },
+	{ "atan2(", ",", ")", "" },
 	{ "pow(", ",", ")", "" }, // ir_binop_pow,
 
 	{ "mix(", ",", ",", ")" }, // ir_ternop_lerp,
