UE4 - Fix bogus warning compiling Metal shaders

diff --git a/Engine/Source/Developer/iOS/MetalShaderFormat/Private/MetalUtils.cpp b/Engine/Source/Developer/iOS/MetalShaderFormat/Private/MetalUtils.cpp
--- a/Engine/Source/Developer/iOS/MetalShaderFormat/Private/MetalUtils.cpp
+++ b/Engine/Source/Developer/iOS/MetalShaderFormat/Private/MetalUtils.cpp
@@ -313,7 +313,7 @@ namespace MetalUtils
 			{
 				Variable->semantic = ralloc_asprintf(ParseState, "[[ attribute(%s) ]]", Semantic);
 			}
-			else
+			else if (FCStringAnsi::Strnicmp(Semantic, "[[", 2))
 			{
 				_mesa_glsl_warning(ParseState, "Unrecognized input attribute '%s'", Semantic);
 			}
