READY - UE4 - Metal bumped to Feature Level ES3.1
- Fix codegen bug

diff --git a/Engine/Source/Developer/iOS/MetalShaderFormat/Private/MetalBackend.cpp b/Engine/Source/Developer/iOS/MetalShaderFormat/Private/MetalBackend.cpp
--- a/Engine/Source/Developer/iOS/MetalShaderFormat/Private/MetalBackend.cpp
+++ b/Engine/Source/Developer/iOS/MetalShaderFormat/Private/MetalBackend.cpp
@@ -964,7 +964,7 @@ protected:
 	virtual void visit(ir_texture *tex) override
 	{
 		check(scope_depth > 0);
-
+		bool bNeedsClosingParenthesis = true;
 		tex->sampler->accept(this);
 		if (tex->op == ir_tex || tex->op == ir_txl || tex->op == ir_txb)
 		{
@@ -996,8 +996,18 @@ protected:
 		}
 		else if (tex->op == ir_txf)
 		{
-			ralloc_asprintf_append(buffer, ".read(");
-			tex->coordinate->accept(this);
+			if (tex->sampler->type && tex->sampler->type->is_sampler() && tex->sampler->type->sampler_buffer)
+			{
+				ralloc_asprintf_append(buffer, "[");
+				tex->coordinate->accept(this);
+				ralloc_asprintf_append(buffer, "]");
+				bNeedsClosingParenthesis = false;
+			}
+			else
+			{
+				ralloc_asprintf_append(buffer, ".read(");
+				tex->coordinate->accept(this);
+			}
 		}
 		else
 		{
@@ -1005,7 +1015,10 @@ protected:
 			check(0);
 		}
 
-		ralloc_asprintf_append(buffer, ")");
+		if (bNeedsClosingParenthesis)
+		{
+			ralloc_asprintf_append(buffer, ")");
+		}
 	}
 
 	virtual void visit(ir_swizzle *swizzle) override
diff --git a/Engine/Source/Developer/iOS/MetalShaderFormat/Private/MetalShaderFormat.cpp b/Engine/Source/Developer/iOS/MetalShaderFormat/Private/MetalShaderFormat.cpp
--- a/Engine/Source/Developer/iOS/MetalShaderFormat/Private/MetalShaderFormat.cpp
+++ b/Engine/Source/Developer/iOS/MetalShaderFormat/Private/MetalShaderFormat.cpp
@@ -15,7 +15,7 @@ class FMetalShaderFormat : public IShaderFormat
 {
 	enum
 	{
-		HEADER_VERSION = 8,
+		HEADER_VERSION = 9,
 	};
 	struct FVersion
 	{
diff --git a/Engine/Source/Runtime/IOS/MetalRHI/Private/MetalRHI.cpp b/Engine/Source/Runtime/IOS/MetalRHI/Private/MetalRHI.cpp
--- a/Engine/Source/Runtime/IOS/MetalRHI/Private/MetalRHI.cpp
+++ b/Engine/Source/Runtime/IOS/MetalRHI/Private/MetalRHI.cpp
@@ -54,7 +54,7 @@ FMetalDynamicRHI::FMetalDynamicRHI()
 	}
 	else
 	{
-		GMaxRHIFeatureLevel = ERHIFeatureLevel::ES2;
+		GMaxRHIFeatureLevel = ERHIFeatureLevel::ES3_1;
 		GMaxRHIShaderPlatform = SP_METAL;
 	}
 
@@ -90,8 +90,8 @@ FMetalDynamicRHI::FMetalDynamicRHI()
 	GMaxCubeTextureDimensions = 4096;
 	GMaxTextureArrayLayers = 2048;
 
-	GShaderPlatformForFeatureLevel[ERHIFeatureLevel::ES2] = SP_METAL;
-	GShaderPlatformForFeatureLevel[ERHIFeatureLevel::ES3_1] = SP_NumPlatforms;
+	GShaderPlatformForFeatureLevel[ERHIFeatureLevel::ES2] = SP_NumPlatforms;
+	GShaderPlatformForFeatureLevel[ERHIFeatureLevel::ES3_1] = SP_METAL;
 	GShaderPlatformForFeatureLevel[ERHIFeatureLevel::SM4] = (GMaxRHIShaderPlatform == SP_METAL_MRT) ? SP_METAL_MRT : SP_NumPlatforms;
 	GShaderPlatformForFeatureLevel[ERHIFeatureLevel::SM5] = SP_NumPlatforms;
 
diff --git a/Engine/Source/Runtime/RHI/Public/RHIDefinitions.h b/Engine/Source/Runtime/RHI/Public/RHIDefinitions.h
--- a/Engine/Source/Runtime/RHI/Public/RHIDefinitions.h
+++ b/Engine/Source/Runtime/RHI/Public/RHIDefinitions.h
@@ -587,8 +587,9 @@ inline ERHIFeatureLevel::Type GetMaxSupportedFeatureLevel(EShaderPlatform InShad
 	case SP_OPENGL_ES2:
 	case SP_OPENGL_ES2_WEBGL:
 	case SP_OPENGL_ES2_IOS:
-	case SP_METAL:
 		return ERHIFeatureLevel::ES2;
+	case SP_METAL:
+		return ERHIFeatureLevel::ES3_1;
 	default:
 		check(0);
 		return ERHIFeatureLevel::Num;
@@ -625,7 +626,7 @@ inline bool IsFeatureLevelSupported(EShaderPlatform InShaderPlatform, ERHIFeatur
 	case SP_XBOXONE:
 		return InFeatureLevel <= ERHIFeatureLevel::SM5;
 	case SP_METAL: 
-		return InFeatureLevel <= ERHIFeatureLevel::ES2;
+		return InFeatureLevel <= ERHIFeatureLevel::ES3_1;
 	case SP_METAL_MRT:
 		return InFeatureLevel <= ERHIFeatureLevel::SM4;
 	case SP_OPENGL_ES31_EXT:
