UE4 - Metal at Feature Level 3.1 (shader side)

diff --git a/Engine/Shaders/Common.usf b/Engine/Shaders/Common.usf
--- a/Engine/Shaders/Common.usf
+++ b/Engine/Shaders/Common.usf
@@ -117,7 +117,7 @@
 		#define class struct
 	#endif
 #elif METAL_PROFILE
-	#define FEATURE_LEVEL FEATURE_LEVEL_ES2
+	#define FEATURE_LEVEL  FEATURE_LEVEL_ES3_1
 	#define FCOLOR_COMPONENT_SWIZZLE .bgra  
 #elif METAL_MRT_PROFILE
 	#define FEATURE_LEVEL FEATURE_LEVEL_SM4
@@ -315,7 +315,7 @@ MaterialFloat PhongShadingPow(MaterialFloat X, MaterialFloat Y)
 	return ClampedPow(X, Y);
 }
 
-#if FEATURE_LEVEL != FEATURE_LEVEL_SM4 && FEATURE_LEVEL != FEATURE_LEVEL_SM5
+#if FEATURE_LEVEL < FEATURE_LEVEL_ES3_1
 	// DX11 (feature levels >= 10) feature sets natively supports uints in shaders; we just use floats on other platforms.
 	#define uint4	int4
 #endif
diff --git a/Engine/Shaders/GpuSkinVertexFactory.usf b/Engine/Shaders/GpuSkinVertexFactory.usf
--- a/Engine/Shaders/GpuSkinVertexFactory.usf
+++ b/Engine/Shaders/GpuSkinVertexFactory.usf
@@ -66,7 +66,7 @@ struct FVertexFactoryInput
 	// TangentZ.w contains sign of tangent basis determinant
 	half4	TangentZ		: ATTRIBUTE2;	
 
-#if FEATURE_LEVEL >= FEATURE_LEVEL_SM4
+#if FEATURE_LEVEL >= FEATURE_LEVEL_ES3_1
 	uint4	BlendIndices	: ATTRIBUTE3;
 	#if GPUSKIN_USE_EXTRA_INFLUENCES
 	uint4	BlendIndicesExtra	: ATTRIBUTE14;
@@ -390,7 +390,7 @@ float3x3 SkinTangents( FVertexFactoryInput Input, FVertexFactoryIntermediates In
 #else
 	half4 TangentZ = TangentBias(Input.TangentZ);
 	
-    #if GPUSKIN_APEX_CLOTH
+	#if GPUSKIN_APEX_CLOTH
 
 	float3x3 ClothTangent;
 	float TempClothBlendWeight = 0.0f;
@@ -421,10 +421,10 @@ float3x3 SkinTangents( FVertexFactoryInput Input, FVertexFactoryIntermediates In
 		TangentToLocal[0] = mul(TangentToLocal[0], (half3x3)Primitive.WorldToLocal);
 		TangentToLocal[2] = mul(TangentToLocal[2], (half3x3)Primitive.WorldToLocal);
 
-	    // derive the new binormal by getting the cross product of the normal and tangent
-	    // and flip vector based on sign of tangent basis determinant
-	    TangentToLocal[1] = cross(TangentToLocal[2], TangentToLocal[0]) * TangentZ.w;
-    
+		// derive the new binormal by getting the cross product of the normal and tangent
+		// and flip vector based on sign of tangent basis determinant
+		TangentToLocal[1] = cross(TangentToLocal[2], TangentToLocal[0]) * TangentZ.w;
+	
 		ClothTangent = TangentToLocal;
 	}
 	#endif // GPUSKIN_APEX_CLOTH
diff --git a/Engine/Source/Developer/iOS/MetalShaderFormat/Private/MetalBackend.cpp b/Engine/Source/Developer/iOS/MetalShaderFormat/Private/MetalBackend.cpp
--- a/Engine/Source/Developer/iOS/MetalShaderFormat/Private/MetalBackend.cpp
+++ b/Engine/Source/Developer/iOS/MetalShaderFormat/Private/MetalBackend.cpp
@@ -23,7 +23,7 @@
 /**
  * This table must match the ir_expression_operation enum.
  */
-static const char * const GLSLExpressionTable[ir_opcode_count][4] =
+static const char * const MetalExpressionTable[ir_opcode_count][4] =
 	{
 	{ "(~", ")", "", "" }, // ir_unop_bit_not,
 	{ "not(", ")", "", "!" }, // ir_unop_logic_not,
@@ -187,7 +187,7 @@ static const char * const GLSLExpressionTable[ir_opcode_count][4] =
 	{ "ERROR_QUADOP_VECTOR(", ",", ")" }, // ir_quadop_vector,
 };
 
-static_assert((sizeof(GLSLExpressionTable) / sizeof(GLSLExpressionTable[0])) == ir_opcode_count, "Metal Expression Table Size Mismatch");
+static_assert((sizeof(MetalExpressionTable) / sizeof(MetalExpressionTable[0])) == ir_opcode_count, "Metal Expression Table Size Mismatch");
 
 struct SDMARange
 {
@@ -916,7 +916,7 @@ protected:
 			 (numOps == 2 && op >= ir_binop_first_comparison && op <= ir_binop_last_comparison) ||
 			 (numOps == 2 && op >= ir_binop_first_logic && op <= ir_binop_last_logic)))
 		{
-			const char* op_str = GLSLExpressionTable[op][3];
+			const char* op_str = MetalExpressionTable[op][3];
 			ralloc_asprintf_append(buffer, "%s(", (numOps == 1) ? op_str : "");
 			expr->operands[0]->accept(this);
 			if (numOps == 2)
@@ -929,11 +929,11 @@ protected:
 		else if (expr->type->is_vector() && numOps == 2 &&
 			op >= ir_binop_first_logic && op <= ir_binop_last_logic)
 		{
-			ralloc_asprintf_append(buffer, GLSLExpressionTable[op][0], expr->type->vector_elements, expr->type->vector_elements);
+			ralloc_asprintf_append(buffer, MetalExpressionTable[op][0], expr->type->vector_elements, expr->type->vector_elements);
 			expr->operands[0]->accept(this);
-			ralloc_asprintf_append(buffer, GLSLExpressionTable[op][1], expr->type->vector_elements);
+			ralloc_asprintf_append(buffer, MetalExpressionTable[op][1], expr->type->vector_elements);
 			expr->operands[1]->accept(this);
-			ralloc_asprintf_append(buffer, GLSLExpressionTable[op][2]);
+			ralloc_asprintf_append(buffer, MetalExpressionTable[op][2]);
 		}
 		else if (op == ir_binop_mod && !expr->type->is_float())
 		{
@@ -952,11 +952,11 @@ protected:
 		}
 		else if (numOps < 4)
 		{
-			ralloc_asprintf_append(buffer, GLSLExpressionTable[op][0]);
+			ralloc_asprintf_append(buffer, MetalExpressionTable[op][0]);
 			for (int i = 0; i < numOps; ++i)
 			{
 				expr->operands[i]->accept(this);
-				ralloc_asprintf_append(buffer, GLSLExpressionTable[op][i+1]);
+				ralloc_asprintf_append(buffer, MetalExpressionTable[op][i+1]);
 			}
 		}
 	}
@@ -2515,78 +2515,6 @@ struct FMetalUnsupportedSamplerVisitor : public ir_hierarchical_visitor
 }; 
 
 
-// Converts an array index expression using an integer input attribute, to a float input attribute using a conversion to int
-struct SConvertIntVertexAttributeES2 : public ir_hierarchical_visitor
-{
-	_mesa_glsl_parse_state* ParseState;
-	exec_list* FunctionBody;
-	int InsideArrayDeref;
-	std::map<ir_variable*, ir_variable*> ConvertedVarMap;
-
-	SConvertIntVertexAttributeES2(_mesa_glsl_parse_state* InParseState, exec_list* InFunctionBody) : ParseState(InParseState), FunctionBody(InFunctionBody), InsideArrayDeref(0)
-	{
-	}
-
-	virtual ir_visitor_status visit_enter(ir_dereference_array* DeRefArray) override
-	{
-		// Break the array dereference so we know we want to modify the array index part
-		auto Result = ir_hierarchical_visitor::visit_enter(DeRefArray);
-		++InsideArrayDeref;
-		DeRefArray->array_index->accept(this);
-		--InsideArrayDeref;
-
-		return visit_continue;
-	}
-
-	virtual ir_visitor_status visit(ir_dereference_variable* DeRefVar) override
-	{
-		if (InsideArrayDeref > 0)
-		{
-			ir_variable* SourceVar = DeRefVar->var;
-			if (SourceVar->mode == ir_var_in)
-			{
-				// First time it still is an integer, so add the temporary and a conversion, and switch to float
-				if (SourceVar->type->is_integer())
-				{
-					check(SourceVar->type->is_integer() && !SourceVar->type->is_matrix() && !SourceVar->type->is_array());
-
-					// Double check we haven't processed this
-					auto IterFound = ConvertedVarMap.find(SourceVar);
-					check(IterFound == ConvertedVarMap.end());
-
-					// New temp var
-					ir_variable* NewVar = new(ParseState)ir_variable(SourceVar->type, NULL, ir_var_temporary);
-					base_ir->insert_before(NewVar);
-
-					// Switch original type to float
-					SourceVar->type = glsl_type::get_instance(GLSL_TYPE_FLOAT, SourceVar->type->vector_elements, 1);
-
-					// Convert float to int
-					ir_dereference_variable* NewSourceDeref = new(ParseState)ir_dereference_variable(SourceVar);
-					ir_expression* NewCastExpression = new(ParseState)ir_expression(ir_unop_f2i, NewSourceDeref);
-					ir_assignment* NewAssigment = new(ParseState)ir_assignment(new(ParseState)ir_dereference_variable(NewVar), NewCastExpression);
-					base_ir->insert_before(NewAssigment);
-
-					// Add the entry and modify the original Var
-					ConvertedVarMap[SourceVar] = NewVar;
-					DeRefVar->var = NewVar;
-				}
-				else
-				{
-					auto IterFound = ConvertedVarMap.find(SourceVar);
-					if (IterFound != ConvertedVarMap.end())
-					{
-						DeRefVar->var = IterFound->second;
-					}
-				}
-			}
-		}
-
-		return ir_hierarchical_visitor::visit(DeRefVar);
-	}
-};
-
-
 bool FMetalCodeBackend::ApplyAndVerifyPlatformRestrictions(exec_list* Instructions, _mesa_glsl_parse_state* ParseState, EHlslShaderFrequency Frequency)
 {
 	bool bIsVertexShader = (Frequency == HSF_VertexShader);
@@ -2598,13 +2526,6 @@ bool FMetalCodeBackend::ApplyAndVerifyPlatformRestrictions(exec_list* Instructio
 		Visitor.run(Instructions);
 	}
 
-	// Handle integer vertex attributes used as array indices
-	if (bIsVertexShader)
-	{
-		SConvertIntVertexAttributeES2 ConvertIntVertexAttributeVisitor(ParseState, Instructions);
-		ConvertIntVertexAttributeVisitor.run(Instructions);
-	}
-
 	return true;
 }
 
diff --git a/Engine/Source/Developer/iOS/MetalShaderFormat/Private/MetalShaderFormat.cpp b/Engine/Source/Developer/iOS/MetalShaderFormat/Private/MetalShaderFormat.cpp
--- a/Engine/Source/Developer/iOS/MetalShaderFormat/Private/MetalShaderFormat.cpp
+++ b/Engine/Source/Developer/iOS/MetalShaderFormat/Private/MetalShaderFormat.cpp
@@ -15,7 +15,7 @@ class FMetalShaderFormat : public IShaderFormat
 {
 	enum
 	{
-		HEADER_VERSION = 9,
+		HEADER_VERSION = 10,
 	};
 	struct FVersion
 	{
