addition of mobile provision and certificate to the IOS project settings
calls IPP for now
UE-4397
#ios

diff --git a/Engine/Source/Developer/iOS/IOSPlatformEditor/IOSPlatformEditor.Build.cs b/Engine/Source/Developer/iOS/IOSPlatformEditor/IOSPlatformEditor.Build.cs
--- a/Engine/Source/Developer/iOS/IOSPlatformEditor/IOSPlatformEditor.Build.cs
+++ b/Engine/Source/Developer/iOS/IOSPlatformEditor/IOSPlatformEditor.Build.cs
@@ -13,7 +13,9 @@ public class IOSPlatformEditor : ModuleRules
 				"Core",
 				"CoreUObject",
 				"InputCore",
+				"DesktopPlatform",
 				"Engine",
+				"MainFrame",
 				"Slate",
 				"SlateCore",
 				"EditorStyle",
@@ -21,6 +23,7 @@ public class IOSPlatformEditor : ModuleRules
 				"SharedSettingsWidgets",
 				"SourceControl",
 				"IOSRuntimeSettings",
+				"TargetPlatform",
 			}
 		);
 
diff --git a/Engine/Source/Developer/iOS/IOSPlatformEditor/Private/IOSTargetSettingsCustomization.cpp b/Engine/Source/Developer/iOS/IOSPlatformEditor/Private/IOSTargetSettingsCustomization.cpp
--- a/Engine/Source/Developer/iOS/IOSPlatformEditor/Private/IOSTargetSettingsCustomization.cpp
+++ b/Engine/Source/Developer/iOS/IOSPlatformEditor/Private/IOSTargetSettingsCustomization.cpp
@@ -5,6 +5,8 @@
 #include "DetailLayoutBuilder.h"
 #include "DetailCategoryBuilder.h"
 #include "PropertyEditing.h"
+#include "DesktopPlatformModule.h"
+#include "MainFrame.h"
 
 #include "ScopedTransaction.h"
 #include "SExternalImageReference.h"
@@ -15,6 +17,7 @@
 #include "ManifestUpdateHelper.h"
 #include "SNotificationList.h"
 #include "NotificationManager.h"
+#include "TargetPlatform.h"
 
 #define LOCTEXT_NAMESPACE "IOSTargetSettings"
 
@@ -62,14 +65,32 @@ void FIOSTargetSettingsCustomization::CustomizeDetails(IDetailLayoutBuilder& Det
 {
 	SavedLayoutBuilder = &DetailLayout;
 
+	UpdateStatus();
+
 	BuildPListSection(DetailLayout);
 
 	BuildIconSection(DetailLayout);
 }
 
+void FIOSTargetSettingsCustomization::UpdateStatus()
+{
+	// get the provision status
+	bProvisionInstalled = bCertificateInstalled = false;
+	const ITargetPlatform* const Platform = GetTargetPlatformManager()->FindTargetPlatform("IOS");
+	if (Platform)
+	{
+		FString NotInstalledTutorialLink;
+		FString ProjectPath = FPaths::IsProjectFilePathSet() ? FPaths::ConvertRelativePathToFull(FPaths::GetProjectFilePath()) : FPaths::RootDir() / FApp::GetGameName() / FApp::GetGameName() + TEXT(".uproject");
+		int32 Result = Platform->CheckRequirements(ProjectPath, false, NotInstalledTutorialLink);
+		bProvisionInstalled = !(Result & ETargetPlatformReadyStatus::ProvisionNotFound);
+		bCertificateInstalled = !(Result & ETargetPlatformReadyStatus::SigningKeyNotFound);
+	}
+}
+
 void FIOSTargetSettingsCustomization::BuildPListSection(IDetailLayoutBuilder& DetailLayout)
 {
 	// Info.plist category
+	IDetailCategoryBuilder& ProvisionCategory = DetailLayout.EditCategory(TEXT("Mobile Provision"));
 	IDetailCategoryBuilder& AppManifestCategory = DetailLayout.EditCategory(TEXT("Info.plist"));
 	IDetailCategoryBuilder& BundleCategory = DetailLayout.EditCategory(TEXT("Bundle Information"));
 	IDetailCategoryBuilder& OrientationCategory = DetailLayout.EditCategory(TEXT("Orientation"));
@@ -84,6 +105,108 @@ void FIOSTargetSettingsCustomization::BuildPListSection(IDetailLayoutBuilder& De
 
 	SetupForPlatformAttribute = PlatformSetupMessage->GetReadyToGoAttribute();
 
+/*	ProvisionCategory.AddCustomRow(TEXT("Certificate Request"), false)
+		.NameContent()
+		[
+			SNew(SHorizontalBox)
+			+ SHorizontalBox::Slot()
+			.Padding(FMargin(0, 1, 0, 1))
+			.FillWidth(1.0f)
+			[
+				SNew(STextBlock)
+				.Text(LOCTEXT("RequestLabel", "Certificate Request"))
+				.Font(DetailLayout.GetDetailFont())
+			]
+		]
+	.ValueContent()
+		[
+			SNew(SBox)
+			.HAlign(HAlign_Center)
+			[
+				SNew(SButton)
+				.HAlign(HAlign_Center)
+				.OnClicked(this, &FIOSTargetSettingsCustomization::OnCertificateRequestClicked)
+				[
+					SNew(STextBlock)
+					.Text(FText::FromString("Create Certificate Request and a Key Pair"))
+				]
+			]
+		];*/
+
+	ProvisionCategory.AddCustomRow(TEXT("Mobile Provision"), false)
+		.NameContent()
+		[
+			SNew(SHorizontalBox)
+			+ SHorizontalBox::Slot()
+			.Padding(FMargin(0, 1, 0, 1))
+			.FillWidth(1.0f)
+			[
+				SNew(STextBlock)
+				.Text(LOCTEXT("ProvisionLabel", "Provision"))
+				.Font(DetailLayout.GetDetailFont())
+			]
+		]
+		.ValueContent()
+		[
+			SNew(SHorizontalBox)
+			+ SHorizontalBox::Slot()
+			.Padding(FMargin(8, 0, 8, 0))
+			.AutoWidth()
+			[
+				SNew(SImage)
+				.Image(this, &FIOSTargetSettingsCustomization::GetProvisionStatus)
+			]
+			+ SHorizontalBox::Slot()
+			.Padding(FMargin(0, 1, 0, 1))
+			.FillWidth(1.0f)
+			[
+				SNew(SButton)
+				.HAlign(HAlign_Center)
+				.OnClicked(this, &FIOSTargetSettingsCustomization::OnInstallProvisionClicked)
+				[
+					SNew(STextBlock)
+					.Text(FText::FromString("Import Provision"))
+				]
+			]
+		];
+
+	ProvisionCategory.AddCustomRow(TEXT("Certificate"), false)
+		.NameContent()
+		[
+			SNew(SHorizontalBox)
+			+ SHorizontalBox::Slot()
+			.Padding(FMargin(0, 1, 0, 1))
+			.FillWidth(1.0f)
+			[
+				SNew(STextBlock)
+				.Text(LOCTEXT("CertificateLabel", "Certificate"))
+				.Font(DetailLayout.GetDetailFont())
+			]
+		]
+		.ValueContent()
+			[
+				SNew(SHorizontalBox)
+				+ SHorizontalBox::Slot()
+				.Padding(FMargin(8, 0, 8, 0))
+				.AutoWidth()
+				[
+					SNew(SImage)
+					.Image(this, &FIOSTargetSettingsCustomization::GetCertificateStatus)
+				]
+				+ SHorizontalBox::Slot()
+				.Padding(FMargin(0, 1, 0, 1))
+				.FillWidth(1.0f)
+				[
+					SNew(SButton)
+					.HAlign(HAlign_Center)
+					.OnClicked(this, &FIOSTargetSettingsCustomization::OnInstallCertificateClicked)
+					[
+						SNew(STextBlock)
+						.Text(FText::FromString("Import Certificate"))
+					]
+				]
+			];
+
 	AppManifestCategory.AddCustomRow(TEXT("Warning"), false)
 		.WholeRowWidget
 		[
@@ -371,6 +494,101 @@ void FIOSTargetSettingsCustomization::BuildImageRow(IDetailLayoutBuilder& Detail
 		];
 }
 
+static FString OutputMessage;
+static void OnOutput(FString Message)
+{
+	OutputMessage += Message;
+	UE_LOG(LogTemp, Display, TEXT("%s\n"), *Message);
+}
+
+FReply FIOSTargetSettingsCustomization::OnInstallProvisionClicked()
+{
+	// pass the file to IPP to install
+	FString ProjectPath = FPaths::IsProjectFilePathSet() ? FPaths::ConvertRelativePathToFull(FPaths::GetProjectFilePath()) : FPaths::RootDir() / FApp::GetGameName() / FApp::GetGameName() + TEXT(".uproject");
+#if PLATFORM_MAC
+	FString CmdExe = TEXT("/bin/sh");
+	FString ScriptPath = FPaths::ConvertRelativePathToFull(FPaths::EngineDir() / TEXT("Build/BatchFiles/Mac/RunMono.sh"));
+	FString IPPPath = FPaths::ConvertRelativePathToFull(FPaths::EngineDir() / TEXT("Binaries/DotNet/IOS/IPhonePackager.exe"));
+	FString CommandLine = FString::Printf(TEXT("\"%s\" \"%s\" Install Engine -project \"%s\" -provision"), *ScriptPath, *IPPPath, *ProjectPath);
+#else
+	FString CmdExe = FPaths::ConvertRelativePathToFull(FPaths::EngineDir() / TEXT("Binaries/DotNet/IOS/IPhonePackager.exe"));
+	FString CommandLine = FString::Printf(TEXT("Install Engine -project \"%s\" -provision"), *ProjectPath);
+#endif
+	TSharedPtr<FMonitoredProcess> IPPProcess = MakeShareable(new FMonitoredProcess(CmdExe, CommandLine, true));
+	OutputMessage = TEXT("");
+	IPPProcess->OnOutput().BindStatic(&OnOutput);
+	IPPProcess->Launch();
+	while(IPPProcess->IsRunning())
+	{
+		FPlatformProcess::Sleep(0.01f);
+	}
+	int RetCode = IPPProcess->GetReturnCode();
+	ensure(RetCode == 0);
+
+	UpdateStatus();
+
+	return FReply::Handled();
+}
+
+FReply FIOSTargetSettingsCustomization::OnInstallCertificateClicked()
+{
+	// pass the file to IPP to install
+	FString ProjectPath = FPaths::IsProjectFilePathSet() ? FPaths::ConvertRelativePathToFull(FPaths::GetProjectFilePath()) : FPaths::RootDir() / FApp::GetGameName() / FApp::GetGameName() + TEXT(".uproject");
+#if PLATFORM_MAC
+	FString CmdExe = TEXT("/bin/sh");
+	FString ScriptPath = FPaths::ConvertRelativePathToFull(FPaths::EngineDir() / TEXT("Build/BatchFiles/Mac/RunMono.sh"));
+	FString IPPPath = FPaths::ConvertRelativePathToFull(FPaths::EngineDir() / TEXT("Binaries/DotNet/IOS/IPhonePackager.exe"));
+	FString CommandLine = FString::Printf(TEXT("\"%s\" \"%s\" Install Engine -project \"%s\" -certificate"), *ScriptPath, *IPPPath, *ProjectPath);
+#else
+	FString CmdExe = FPaths::ConvertRelativePathToFull(FPaths::EngineDir() / TEXT("Binaries/DotNet/IOS/IPhonePackager.exe"));
+	FString CommandLine = FString::Printf(TEXT("Install Engine -project \"%s\" -certificate"), *ProjectPath);
+#endif
+	TSharedPtr<FMonitoredProcess> IPPProcess = MakeShareable(new FMonitoredProcess(CmdExe, CommandLine, true));
+	OutputMessage = TEXT("");
+	IPPProcess->OnOutput().BindStatic(&OnOutput);
+	IPPProcess->Launch();
+	while(IPPProcess->IsRunning())
+	{
+		FPlatformProcess::Sleep(0.01f);
+	}
+	int RetCode = IPPProcess->GetReturnCode();
+	ensure(RetCode == 0);
+
+	UpdateStatus();
+
+	return FReply::Handled();
+}
+
+FReply FIOSTargetSettingsCustomization::OnCertificateRequestClicked()
+{
+	// TODO: bring up an open file dialog and then install the provision
+	return FReply::Handled();
+}
+
+const FSlateBrush* FIOSTargetSettingsCustomization::GetProvisionStatus() const
+{
+	if( bProvisionInstalled )
+	{
+		return FEditorStyle::GetBrush("Automation.Success");
+	}
+	else
+	{
+		return FEditorStyle::GetBrush("Automation.Fail");
+	}
+}
+
+const FSlateBrush* FIOSTargetSettingsCustomization::GetCertificateStatus() const
+{
+	if( bCertificateInstalled )
+	{
+		return FEditorStyle::GetBrush("Automation.Success");
+	}
+	else
+	{
+		return FEditorStyle::GetBrush("Automation.Fail");
+	}
+}
+
 //////////////////////////////////////////////////////////////////////////
 
 #undef LOCTEXT_NAMESPACE
\ No newline at end of file
diff --git a/Engine/Source/Developer/iOS/IOSPlatformEditor/Private/IOSTargetSettingsCustomization.h b/Engine/Source/Developer/iOS/IOSPlatformEditor/Private/IOSTargetSettingsCustomization.h
--- a/Engine/Source/Developer/iOS/IOSPlatformEditor/Private/IOSTargetSettingsCustomization.h
+++ b/Engine/Source/Developer/iOS/IOSPlatformEditor/Private/IOSTargetSettingsCustomization.h
@@ -34,6 +34,9 @@ private:
 	// Is the manifest writable?
 	TAttribute<bool> SetupForPlatformAttribute;
 
+	bool bProvisionInstalled;
+	bool bCertificateInstalled;
+
 private:
 	FIOSTargetSettingsCustomization();
 
@@ -51,4 +54,22 @@ private:
 
 	// Builds an image row
 	void BuildImageRow(class IDetailLayoutBuilder& DetailLayout, class IDetailCategoryBuilder& Category, const struct FPlatformIconInfo& Info, const FVector2D& MaxDisplaySize);
+
+	// Install the provision
+ 	FReply OnInstallProvisionClicked();
+
+	// Install the provision
+	FReply OnInstallCertificateClicked();
+
+	// certificate request
+	FReply OnCertificateRequestClicked();
+
+	// Get the image to display for the provision status
+	const FSlateBrush* GetProvisionStatus() const;
+
+	// Get the image to display for the certificate status
+	const FSlateBrush* GetCertificateStatus() const;
+
+	// Update the provision status
+	void UpdateStatus();
 };
diff --git a/Engine/Source/Programs/IOS/iPhonePackager/Config.cs b/Engine/Source/Programs/IOS/iPhonePackager/Config.cs
--- a/Engine/Source/Programs/IOS/iPhonePackager/Config.cs
+++ b/Engine/Source/Programs/IOS/iPhonePackager/Config.cs
@@ -286,6 +286,16 @@ namespace iPhonePackager
 		/// </summary>
 		public static bool bCookOnTheFly = false;
 
+		/// <summary>
+		/// Whether the install should be performed on a provision
+		/// </summary>
+		public static bool bProvision = false;
+
+		/// <summary>
+		/// Whether the install should be performed on a certificate
+		/// </summary>
+		public static bool bCert = false;
+
 		/// <summary>
 		/// An override server Mac name
 		/// </summary>
diff --git a/Engine/Source/Programs/IOS/iPhonePackager/iPhonePackager.cs b/Engine/Source/Programs/IOS/iPhonePackager/iPhonePackager.cs
--- a/Engine/Source/Programs/IOS/iPhonePackager/iPhonePackager.cs
+++ b/Engine/Source/Programs/IOS/iPhonePackager/iPhonePackager.cs
@@ -224,6 +224,12 @@ namespace iPhonePackager
 							case "-cookonthefly":
 								Config.bCookOnTheFly = true;
 								break;
+							case "-provision":
+								Config.bProvision = true;
+								break;
+							case "-certificate":
+								Config.bCert = true;
+								break;
 						}
 
 						// get the stage dir path
@@ -560,6 +566,7 @@ namespace iPhonePackager
 					Log(" ... RepackageFromStage GameName");
 					Log(" ... Devices");
 					Log(" ... Validate");
+					Log(" ... Install");
 					Log("");
 					Log("Configuration switches:");
 					Log("	 -stagedir <path>		  sets the directory to copy staged files from (defaults to none)");
@@ -704,6 +711,18 @@ namespace iPhonePackager
 						}
 						break;
 
+					case "install":
+						GameName = "";
+						if (Config.bProvision)
+						{
+							ToolsHub.TryInstallingMobileProvision();
+						}
+						if (Config.bCert)
+						{
+							ToolsHub.TryInstallingCertificate_PromptForKey();
+						}
+						break;
+
 					case "resigntool":
 						RunInVisualMode(delegate { return new GraphicalResignTool(); });
 						break;
