fix for crash on Mac when installing provision
improved certificate detection on Mac
UE-4919
#ios

diff --git a/Engine/Source/Developer/iOS/IOSPlatformEditor/Private/IOSTargetSettingsCustomization.cpp b/Engine/Source/Developer/iOS/IOSPlatformEditor/Private/IOSTargetSettingsCustomization.cpp
--- a/Engine/Source/Developer/iOS/IOSPlatformEditor/Private/IOSTargetSettingsCustomization.cpp
+++ b/Engine/Source/Developer/iOS/IOSPlatformEditor/Private/IOSTargetSettingsCustomization.cpp
@@ -505,27 +505,84 @@ FReply FIOSTargetSettingsCustomization::OnInstallProvisionClicked()
 {
 	// pass the file to IPP to install
 	FString ProjectPath = FPaths::IsProjectFilePathSet() ? FPaths::ConvertRelativePathToFull(FPaths::GetProjectFilePath()) : FPaths::RootDir() / FApp::GetGameName() / FApp::GetGameName() + TEXT(".uproject");
+	FString ProvisionPath;
+
+	// get the provision by popping up the file dialog
+	TArray<FString> OpenFilenames;
+	IDesktopPlatform* DesktopPlatform = FDesktopPlatformModule::Get();
+	bool bOpened = false;
+	int32 FilterIndex = -1;
+	FString FileTypes = TEXT("Provision Files (*.mobileprovision)|*.mobileprovision");
+
+	if ( DesktopPlatform )
+	{
+		void* ParentWindowWindowHandle = NULL;
+
+		IMainFrameModule& MainFrameModule = FModuleManager::LoadModuleChecked<IMainFrameModule>(TEXT("MainFrame"));
+		const TSharedPtr<SWindow>& MainFrameParentWindow = MainFrameModule.GetParentWindow();
+		if ( MainFrameParentWindow.IsValid() && MainFrameParentWindow->GetNativeWindow().IsValid() )
+		{
+			ParentWindowWindowHandle = MainFrameParentWindow->GetNativeWindow()->GetOSWindowHandle();
+		}
+
+		bOpened = DesktopPlatform->OpenFileDialog(
+			ParentWindowWindowHandle,
+			LOCTEXT("ImportDialogTitle", "Import Provision").ToString(),
+			FPaths::GetProjectFilePath(),
+			TEXT(""),
+			FileTypes,
+			EFileDialogFlags::None,
+			OpenFilenames,
+			FilterIndex
+			);
+	}
+
+	if ( bOpened )
+	{
+        ProvisionPath = FPaths::ConvertRelativePathToFull(OpenFilenames[0]);
+
+		// see if the provision is already installed
+		FString DestName = FPaths::GetBaseFilename(ProvisionPath);
+		TCHAR Path[4096];
 #if PLATFORM_MAC
-	FString CmdExe = TEXT("/bin/sh");
-	FString ScriptPath = FPaths::ConvertRelativePathToFull(FPaths::EngineDir() / TEXT("Build/BatchFiles/Mac/RunMono.sh"));
-	FString IPPPath = FPaths::ConvertRelativePathToFull(FPaths::EngineDir() / TEXT("Binaries/DotNet/IOS/IPhonePackager.exe"));
-	FString CommandLine = FString::Printf(TEXT("\"%s\" \"%s\" Install Engine -project \"%s\" -provision"), *ScriptPath, *IPPPath, *ProjectPath);
+		FPlatformMisc::GetEnvironmentVariable(TEXT("HOME"), Path, ARRAY_COUNT(Path));
+		FString Destination = FString::Printf(TEXT("\"%s/Library/MobileDevice/Provisioning Profiles/%s.mobileprovision\""), Path, *DestName);
+		FString Destination2 = FString::Printf(TEXT("\"%s/Library/MobileDevice/Provisioning Profiles/%s.mobileprovision\""), Path, FApp::GetGameName());
 #else
-	FString CmdExe = FPaths::ConvertRelativePathToFull(FPaths::EngineDir() / TEXT("Binaries/DotNet/IOS/IPhonePackager.exe"));
-	FString CommandLine = FString::Printf(TEXT("Install Engine -project \"%s\" -provision"), *ProjectPath);
+		FPlatformMisc::GetEnvironmentVariable(TEXT("LOCALAPPDATA"), Path, ARRAY_COUNT(Path));
+		FString Destination = FString::Printf(TEXT("%s\\Apple Computer\\MobileDevice\\Provisioning Profiles\\%s.mobileprovision"), Path, *DestName);
+		FString Destination2 = FString::Printf(TEXT("%s\\Apple Computer\\MobileDevice\\Provisioning Profiles\\%s.mobileprovision"), Path, FApp::GetGameName());
 #endif
-	TSharedPtr<FMonitoredProcess> IPPProcess = MakeShareable(new FMonitoredProcess(CmdExe, CommandLine, true));
-	OutputMessage = TEXT("");
-	IPPProcess->OnOutput().BindStatic(&OnOutput);
-	IPPProcess->Launch();
-	while(IPPProcess->IsRunning())
-	{
-		FPlatformProcess::Sleep(0.01f);
-	}
-	int RetCode = IPPProcess->GetReturnCode();
-	ensure(RetCode == 0);
+		if (FPaths::FileExists(Destination) || FPaths::FileExists(Destination2))
+		{
+			FString MessagePrompt = FString::Printf(TEXT("%s mobile provision file already exists.  Do you want to replace this provision?"), *DestName);
+			if (FPlatformMisc::MessageBoxExt(EAppMsgType::OkCancel, *MessagePrompt, TEXT("File Exists")) == EAppReturnType::Cancel)
+			{
+				return FReply::Handled();
+			}
+		}
+#if PLATFORM_MAC
+		FString CmdExe = TEXT("/bin/sh");
+		FString ScriptPath = FPaths::ConvertRelativePathToFull(FPaths::EngineDir() / TEXT("Build/BatchFiles/Mac/RunMono.sh"));
+		FString IPPPath = FPaths::ConvertRelativePathToFull(FPaths::EngineDir() / TEXT("Binaries/DotNet/IOS/IPhonePackager.exe"));
+		FString CommandLine = FString::Printf(TEXT("\"%s\" \"%s\" Install Engine -project \"%s\" -provision \"%s\""), *ScriptPath, *IPPPath, *ProjectPath, *ProvisionPath);
+#else
+		FString CmdExe = FPaths::ConvertRelativePathToFull(FPaths::EngineDir() / TEXT("Binaries/DotNet/IOS/IPhonePackager.exe"));
+		FString CommandLine = FString::Printf(TEXT("Install Engine -project \"%s\" -provision \"%s\""), *ProjectPath, *ProvisionPath);
+#endif
+		TSharedPtr<FMonitoredProcess> IPPProcess = MakeShareable(new FMonitoredProcess(CmdExe, CommandLine, true));
+		OutputMessage = TEXT("");
+		IPPProcess->OnOutput().BindStatic(&OnOutput);
+		IPPProcess->Launch();
+		while(IPPProcess->IsRunning())
+		{
+			FPlatformProcess::Sleep(0.01f);
+		}
+		int RetCode = IPPProcess->GetReturnCode();
+		ensure(RetCode == 0);
 
-	UpdateStatus();
+		UpdateStatus();
+	}
 
 	return FReply::Handled();
 }
@@ -534,27 +591,63 @@ FReply FIOSTargetSettingsCustomization::OnInstallCertificateClicked()
 {
 	// pass the file to IPP to install
 	FString ProjectPath = FPaths::IsProjectFilePathSet() ? FPaths::ConvertRelativePathToFull(FPaths::GetProjectFilePath()) : FPaths::RootDir() / FApp::GetGameName() / FApp::GetGameName() + TEXT(".uproject");
+	FString CertPath;
+
+	// get the provision by popping up the file dialog
+	TArray<FString> OpenFilenames;
+	IDesktopPlatform* DesktopPlatform = FDesktopPlatformModule::Get();
+	bool bOpened = false;
+	int32 FilterIndex = -1;
+	FString FileTypes = TEXT("Code Signing Certificates (*.cer;*.p12)|*.cer;*p12");
+
+	if ( DesktopPlatform )
+	{
+		void* ParentWindowWindowHandle = NULL;
+
+		IMainFrameModule& MainFrameModule = FModuleManager::LoadModuleChecked<IMainFrameModule>(TEXT("MainFrame"));
+		const TSharedPtr<SWindow>& MainFrameParentWindow = MainFrameModule.GetParentWindow();
+		if ( MainFrameParentWindow.IsValid() && MainFrameParentWindow->GetNativeWindow().IsValid() )
+		{
+			ParentWindowWindowHandle = MainFrameParentWindow->GetNativeWindow()->GetOSWindowHandle();
+		}
+
+		bOpened = DesktopPlatform->OpenFileDialog(
+			ParentWindowWindowHandle,
+			LOCTEXT("ImportDialogTitle", "Import Certificate").ToString(),
+			FPaths::GetProjectFilePath(),
+			TEXT(""),
+			FileTypes,
+			EFileDialogFlags::None,
+			OpenFilenames,
+			FilterIndex
+			);
+	}
+
+	if ( bOpened )
+	{
+        CertPath = FPaths::ConvertRelativePathToFull(OpenFilenames[0]);
 #if PLATFORM_MAC
-	FString CmdExe = TEXT("/bin/sh");
-	FString ScriptPath = FPaths::ConvertRelativePathToFull(FPaths::EngineDir() / TEXT("Build/BatchFiles/Mac/RunMono.sh"));
-	FString IPPPath = FPaths::ConvertRelativePathToFull(FPaths::EngineDir() / TEXT("Binaries/DotNet/IOS/IPhonePackager.exe"));
-	FString CommandLine = FString::Printf(TEXT("\"%s\" \"%s\" Install Engine -project \"%s\" -certificate"), *ScriptPath, *IPPPath, *ProjectPath);
+		FString CmdExe = TEXT("/bin/sh");
+		FString ScriptPath = FPaths::ConvertRelativePathToFull(FPaths::EngineDir() / TEXT("Build/BatchFiles/Mac/RunMono.sh"));
+		FString IPPPath = FPaths::ConvertRelativePathToFull(FPaths::EngineDir() / TEXT("Binaries/DotNet/IOS/IPhonePackager.exe"));
+		FString CommandLine = FString::Printf(TEXT("\"%s\" \"%s\" Install Engine -project \"%s\" -certificate \"%s\""), *ScriptPath, *IPPPath, *ProjectPath, *CertPath);
 #else
-	FString CmdExe = FPaths::ConvertRelativePathToFull(FPaths::EngineDir() / TEXT("Binaries/DotNet/IOS/IPhonePackager.exe"));
-	FString CommandLine = FString::Printf(TEXT("Install Engine -project \"%s\" -certificate"), *ProjectPath);
+		FString CmdExe = FPaths::ConvertRelativePathToFull(FPaths::EngineDir() / TEXT("Binaries/DotNet/IOS/IPhonePackager.exe"));
+		FString CommandLine = FString::Printf(TEXT("Install Engine -project \"%s\" -certificate \"%s\""), *ProjectPath, *CertPath);
 #endif
-	TSharedPtr<FMonitoredProcess> IPPProcess = MakeShareable(new FMonitoredProcess(CmdExe, CommandLine, true));
-	OutputMessage = TEXT("");
-	IPPProcess->OnOutput().BindStatic(&OnOutput);
-	IPPProcess->Launch();
-	while(IPPProcess->IsRunning())
-	{
-		FPlatformProcess::Sleep(0.01f);
-	}
-	int RetCode = IPPProcess->GetReturnCode();
-	ensure(RetCode == 0);
+		TSharedPtr<FMonitoredProcess> IPPProcess = MakeShareable(new FMonitoredProcess(CmdExe, CommandLine, true));
+		OutputMessage = TEXT("");
+		IPPProcess->OnOutput().BindStatic(&OnOutput);
+		IPPProcess->Launch();
+		while(IPPProcess->IsRunning())
+		{
+			FPlatformProcess::Sleep(0.01f);
+		}
+		int RetCode = IPPProcess->GetReturnCode();
+		ensure(RetCode == 0);
 
-	UpdateStatus();
+		UpdateStatus();
+	}
 
 	return FReply::Handled();
 }
diff --git a/Engine/Source/Programs/IOS/iPhonePackager/CodeSigning.cs b/Engine/Source/Programs/IOS/iPhonePackager/CodeSigning.cs
--- a/Engine/Source/Programs/IOS/iPhonePackager/CodeSigning.cs
+++ b/Engine/Source/Programs/IOS/iPhonePackager/CodeSigning.cs
@@ -221,70 +221,57 @@ namespace iPhonePackager
 
 			X509Certificate2 Result = null;
 
-			if (Environment.OSVersion.Platform == PlatformID.Unix || Environment.OSVersion.Platform == PlatformID.MacOSX) {
+			if (Environment.OSVersion.Platform == PlatformID.Unix || Environment.OSVersion.Platform == PlatformID.MacOSX)
+			{
 				// run certtool y to get the currently installed certificates
 				CertToolData = "";
 				Process CertTool = new Process ();
-				CertTool.StartInfo.FileName = "/usr/bin/certtool";
+				CertTool.StartInfo.FileName = "/usr/bin/security";
 				CertTool.StartInfo.UseShellExecute = false;
-				CertTool.StartInfo.Arguments = "y";
+				CertTool.StartInfo.Arguments = "find-identity -p codesigning -v";
 				CertTool.StartInfo.RedirectStandardOutput = true;
 				CertTool.OutputDataReceived += new DataReceivedEventHandler (OutputReceivedCertToolProcessCall);
 				CertTool.Start ();
 				CertTool.BeginOutputReadLine ();
 				CertTool.WaitForExit ();
-				if (CertTool.ExitCode == 0) {
-					foreach (X509Certificate2 SourceCert in ProvisionToWorkFrom.DeveloperCertificates) {
+				if (CertTool.ExitCode == 0)
+				{
+					foreach (X509Certificate2 SourceCert in ProvisionToWorkFrom.DeveloperCertificates)
+					{
 						X509Certificate2 ValidInTimeCert = null;
 						// see if certificate can be found by serial number
-						string SerialNumber = SourceCert.GetSerialNumberString ();
-						for (int Index = 14; Index > 0; Index -= 2) {
-							SerialNumber = SerialNumber.Insert (Index, " ");
-						}
+						string CertHash = SourceCert.GetCertHashString();
 
-						if (CertToolData.Contains (SerialNumber)) {
-							// check the cert time
-							int StartLine = CertToolData.IndexOf (SerialNumber);
-							string BeforeStart = CertToolData.Substring( CertToolData.IndexOf( ":", CertToolData.IndexOf ("Not Before", StartLine))+1);
-							BeforeStart = BeforeStart.Remove(BeforeStart.IndexOf("\n"));
-							string AfterStart = CertToolData.Substring( CertToolData.IndexOf( ":", CertToolData.IndexOf ("Not After", StartLine))+1);
-							AfterStart = AfterStart.Remove(AfterStart.IndexOf("\n"));
-							string CommonName = CertToolData.Substring( CertToolData.IndexOf( ":", CertToolData.IndexOf ("Common Name", CertToolData.IndexOf("Subject Name", StartLine)))+1);
-							CommonName = CommonName.Remove(CommonName.IndexOf("\n"));
-
-							DateTime EffectiveDate = DateTime.Parse (BeforeStart);
-							DateTime ExpirationDate = DateTime.Parse (AfterStart);
-							DateTime Now = DateTime.Now;
-
-							bool bCertTimeIsValid = (EffectiveDate < Now) && (ExpirationDate > Now);
-
-							Program.LogVerbose ("  .. .. Installed certificate '{0}' is {1} (range '{2}' to '{3}')", CommonName, bCertTimeIsValid ? "valid (choosing it)" : "EXPIRED", BeforeStart, AfterStart);
-							if (bCertTimeIsValid) {
-								ValidInTimeCert = SourceCert;
-							}
+						if (CertToolData.Contains (CertHash))
+						{
+							ValidInTimeCert = SourceCert;
 						}
 
-						if (ValidInTimeCert != null) {
+						if (ValidInTimeCert != null)
+						{
 							// Found a cert in the valid time range, quit now!
 							Result = ValidInTimeCert;
 							break;
 						}
 					}
-				} else {
 				}
-			} else {
+			}
+			else
+			{
 				// Open the personal certificate store on this machine
 				X509Store Store = new X509Store ();
 				Store.Open (OpenFlags.ReadOnly);
 
 				// Try finding a matching certificate from the serial number (the one in the mobileprovision is missing the public/private key pair)
-				foreach (X509Certificate2 SourceCert in ProvisionToWorkFrom.DeveloperCertificates) {
+				foreach (X509Certificate2 SourceCert in ProvisionToWorkFrom.DeveloperCertificates)
+				{
 					X509Certificate2Collection FoundCerts = Store.Certificates.Find (X509FindType.FindBySerialNumber, SourceCert.SerialNumber, false);
 
 					Program.LogVerbose ("  .. Provision entry SN '{0}' matched {1} installed certificate(s)", SourceCert.SerialNumber, FoundCerts.Count);
 
 					X509Certificate2 ValidInTimeCert = null;
-					foreach (X509Certificate2 TestCert in FoundCerts) {
+					foreach (X509Certificate2 TestCert in FoundCerts)
+					{
 						//@TODO: Pretty sure the certificate information from the library is in local time, not UTC and this works as expected, but it should be verified!
 						DateTime EffectiveDate = TestCert.NotBefore;
 						DateTime ExpirationDate = TestCert.NotAfter;
@@ -293,13 +280,15 @@ namespace iPhonePackager
 						bool bCertTimeIsValid = (EffectiveDate < Now) && (ExpirationDate > Now);
 
 						Program.LogVerbose ("  .. .. Installed certificate '{0}' is {1} (range '{2}' to '{3}')", TestCert.FriendlyName, bCertTimeIsValid ? "valid (choosing it)" : "EXPIRED", TestCert.GetEffectiveDateString (), TestCert.GetExpirationDateString ());
-						if (bCertTimeIsValid) {
+						if (bCertTimeIsValid)
+						{
 							ValidInTimeCert = TestCert;
 							break;
 						}
 					}
 
-					if (ValidInTimeCert != null) {
+					if (ValidInTimeCert != null)
+					{
 						// Found a cert in the valid time range, quit now!
 						Result = ValidInTimeCert;
 						break;
@@ -561,20 +550,23 @@ namespace iPhonePackager
 			FatBinaryFile FatBinary = new FatBinaryFile();
 			FatBinary.LoadFromBytes(SourceExeData);
 
+			//@TODO: Verify it's an executable (not an object file, etc...)
+			ulong CurrentStreamOffset = 0;
+			byte[] FinalExeData = new byte[SourceExeData.Length + 1024 * 1024];
+
 			foreach (MachObjectFile Exe in FatBinary.MachObjectFiles)
 			{
 				Program.Log("... Processing one mach object (binary is {0})", FatBinary.bIsFatBinary ? "fat" : "thin");
-				
+
 				// Pad the memory stream with extra room to handle any possible growth in the code signing data
 				int OverSize = 1024 * 1024;
-				MemoryStream OutputExeStream = new MemoryStream(SourceExeData.Length + OverSize);
-				
+				MemoryStream OutputExeStream = new MemoryStream(SourceExeData.Length - (int)CurrentStreamOffset + OverSize);
+
 				// Copy the executable into the stream
 				OutputExeStream.Seek(0, SeekOrigin.Begin);
-				OutputExeStream.Write(SourceExeData, 0, SourceExeData.Length);
+				OutputExeStream.Write(SourceExeData, (int)CurrentStreamOffset, SourceExeData.Length - (int)CurrentStreamOffset);
 				OutputExeStream.Seek(0, SeekOrigin.Begin);
-
-				//@TODO: Verify it's an executable (not an object file, etc...)
+				long Length = OutputExeStream.Length;
 
 				// Find out if there was an existing code sign blob and find the linkedit segment command
 				MachLoadCommandCodeSignature CodeSigningBlobLC = null;
@@ -711,18 +703,23 @@ namespace iPhonePackager
 
 				Program.Log("... Truncating/copying final binary", DateTime.Now - SigningTime);
 				ulong DesiredExecutableLength = LinkEditSegmentLC.FileSize + LinkEditSegmentLC.FileOffset;
-				byte[] FinalExeData = OutputExeStream.ToArray();
-				if ((ulong)FinalExeData.Length < DesiredExecutableLength)
+
+				if ((ulong)Length < DesiredExecutableLength)
 				{
 					throw new InvalidDataException("Data written is smaller than expected, unable to finish signing process");
 				}
-				Array.Resize(ref FinalExeData, (int)DesiredExecutableLength); //@todo: Extend the file system interface so we don't have to copy 20 MB just to truncate a few hundred bytes
-
-				// Save the patched and signed executable
-				Program.Log("Saving signed executable... ({0:0.00} s elapsed so far)", (DateTime.Now - SigningTime).TotalSeconds);
-				FileSystem.WriteAllBytes(CFBundleExecutable, FinalExeData);
+				byte[] Data = OutputExeStream.ToArray();
+				Data.CopyTo(FinalExeData, (long)CurrentStreamOffset);
+				CurrentStreamOffset += DesiredExecutableLength;
 			}
 
+			// resize to the finale size
+			Array.Resize(ref FinalExeData, (int)CurrentStreamOffset); //@todo: Extend the file system interface so we don't have to copy 20 MB just to truncate a few hundred bytes
+
+			// Save the patched and signed executable
+			Program.Log("Saving signed executable... ({0:0.00} s elapsed so far)", (DateTime.Now - SigningTime).TotalSeconds);
+			FileSystem.WriteAllBytes(CFBundleExecutable, FinalExeData);
+
 			Program.Log("Finished code signing, which took {0:0.00} s", (DateTime.Now - SigningTime).TotalSeconds);
 		}
 	}
diff --git a/Engine/Source/Programs/IOS/iPhonePackager/Config.cs b/Engine/Source/Programs/IOS/iPhonePackager/Config.cs
--- a/Engine/Source/Programs/IOS/iPhonePackager/Config.cs
+++ b/Engine/Source/Programs/IOS/iPhonePackager/Config.cs
@@ -291,11 +291,21 @@ namespace iPhonePackager
 		/// </summary>
 		public static bool bProvision = false;
 
+		/// <summary>
+		/// provision to be installed
+		/// </summary>
+		public static string Provision = "";
+
 		/// <summary>
 		/// Whether the install should be performed on a certificate
 		/// </summary>
 		public static bool bCert = false;
 
+		/// <summary>
+		/// Certificate to be installed
+		/// </summary>
+		public static string Certificate = "";
+
 		/// <summary>
 		/// An override server Mac name
 		/// </summary>
diff --git a/Engine/Source/Programs/IOS/iPhonePackager/GraphicalResignTool.cs b/Engine/Source/Programs/IOS/iPhonePackager/GraphicalResignTool.cs
--- a/Engine/Source/Programs/IOS/iPhonePackager/GraphicalResignTool.cs
+++ b/Engine/Source/Programs/IOS/iPhonePackager/GraphicalResignTool.cs
@@ -403,13 +403,13 @@ namespace iPhonePackager
 
 		private void ImportCertificateButton3_Click(object sender, EventArgs e)
 		{
-			ToolsHub.TryInstallingCertificate_PromptForKey();
+			ToolsHub.TryInstallingCertificate_PromptForKey(null);
 			RefreshProvisionsList();
 		}
 
 		private void ImportProvisionButton3_Click(object sender, EventArgs e)
 		{
-			ToolsHub.TryInstallingMobileProvision();
+			ToolsHub.TryInstallingMobileProvision(null);
 			RefreshProvisionsList();
 		}
 
diff --git a/Engine/Source/Programs/IOS/iPhonePackager/ToolsHub.cs b/Engine/Source/Programs/IOS/iPhonePackager/ToolsHub.cs
--- a/Engine/Source/Programs/IOS/iPhonePackager/ToolsHub.cs
+++ b/Engine/Source/Programs/IOS/iPhonePackager/ToolsHub.cs
@@ -150,10 +150,9 @@ namespace iPhonePackager
 			return CryptoAdapter.GetCommonNameFromCert(Provision.DeveloperCertificates[0]).IndexOf("iPhone Distribution", StringComparison.InvariantCultureIgnoreCase) >= 0;
 		}
 
-		public static void TryInstallingMobileProvision()
+		public static void TryInstallingMobileProvision(string ProvisionFilename, bool ShowPrompt = true)
 		{
-			string ProvisionFilename;
-			if (ShowOpenFileDialog(MobileProvisionFilter, "Choose a mobile provision to install", "mobileprovision", "", ref ChoosingFilesToInstallDirectory, out ProvisionFilename))
+			if (!String.IsNullOrEmpty(ProvisionFilename) || ShowOpenFileDialog(MobileProvisionFilter, "Choose a mobile provision to install", "mobileprovision", "", ref ChoosingFilesToInstallDirectory, out ProvisionFilename))
 			{
 				try
 				{
@@ -180,7 +179,7 @@ namespace iPhonePackager
 							OldProvision.ProvisionName,
 							Provision.ProvisionName);
 
-						if (MessageBox.Show(MessagePrompt, Config.AppDisplayName, MessageBoxButtons.YesNo, MessageBoxIcon.Question) == DialogResult.No)
+						if (ShowPrompt && MessageBox.Show(MessagePrompt, Config.AppDisplayName, MessageBoxButtons.YesNo, MessageBoxIcon.Question) == DialogResult.No)
 						{
 							return;
 						}
@@ -200,72 +199,102 @@ namespace iPhonePackager
 
 		private void ImportMobileProvisionButton_Click(object sender, EventArgs e)
 		{
-			TryInstallingMobileProvision();
+			TryInstallingMobileProvision(null);
 			UpdateStatus();
 		}
 
-		public static void TryInstallingCertificate_PromptForKey()
+		static private string CertToolData = "";
+		static public void OutputReceivedCertToolProcessCall(Object Sender, DataReceivedEventArgs Line)
+		{
+			if ((Line != null) && !String.IsNullOrEmpty(Line.Data))
+			{
+				CertToolData += Line.Data + "\n";
+			}
+		}
+
+		public static void TryInstallingCertificate_PromptForKey(string CertificateFilename, bool ShowPrompt = true)
 		{
 			try
 			{
-				string CertificateFilename;
-				if (ShowOpenFileDialog(CertificatesFilter, "Choose a code signing certificate to import", "", "", ref ChoosingFilesToInstallDirectory, out CertificateFilename))
+				if (!String.IsNullOrEmpty(CertificateFilename) || ShowOpenFileDialog(CertificatesFilter, "Choose a code signing certificate to import", "", "", ref ChoosingFilesToInstallDirectory, out CertificateFilename))
 				{
-					// Load the certificate
-					string CertificatePassword = "";
-					X509Certificate2 Cert = null;
-					try
+					if (Environment.OSVersion.Platform == PlatformID.MacOSX || Environment.OSVersion.Platform == PlatformID.Unix)
 					{
-						Cert = new X509Certificate2(CertificateFilename, CertificatePassword, X509KeyStorageFlags.PersistKeySet | X509KeyStorageFlags.Exportable | X509KeyStorageFlags.MachineKeySet);
+						// run certtool y to get the currently installed certificates
+						CertToolData = "";
+						Process CertTool = new Process();
+						CertTool.StartInfo.FileName = "/usr/bin/security";
+						CertTool.StartInfo.UseShellExecute = false;
+						CertTool.StartInfo.Arguments = "import \"" + CertificateFilename +"\" -k login.keychain";
+						CertTool.StartInfo.RedirectStandardOutput = true;
+						CertTool.OutputDataReceived += new DataReceivedEventHandler(OutputReceivedCertToolProcessCall);
+						CertTool.Start();
+						CertTool.BeginOutputReadLine();
+						CertTool.WaitForExit();
+						if (CertTool.ExitCode != 0)
+						{
+							// todo: provide some feedback that it failed
+						}
+						Console.Write(CertToolData);
 					}
-					catch (System.Security.Cryptography.CryptographicException ex)
+					else
 					{
-						// Try once with a password
-						if (PasswordDialog.RequestPassword(out CertificatePassword))
+						// Load the certificate
+						string CertificatePassword = "";
+						X509Certificate2 Cert = null;
+						try
 						{
 							Cert = new X509Certificate2(CertificateFilename, CertificatePassword, X509KeyStorageFlags.PersistKeySet | X509KeyStorageFlags.Exportable | X509KeyStorageFlags.MachineKeySet);
 						}
-						else
+						catch (System.Security.Cryptography.CryptographicException ex)
 						{
-							// User cancelled dialog, rethrow
-							throw ex;
+							// Try once with a password
+							if (PasswordDialog.RequestPassword(out CertificatePassword))
+							{
+								Cert = new X509Certificate2(CertificateFilename, CertificatePassword, X509KeyStorageFlags.PersistKeySet | X509KeyStorageFlags.Exportable | X509KeyStorageFlags.MachineKeySet);
+							}
+							else
+							{
+								// User cancelled dialog, rethrow
+								throw ex;
+							}
 						}
-					}
-
-					// If the certificate doesn't have a private key pair, ask the user to provide one
-					if (!Cert.HasPrivateKey)
-					{
-						string ErrorMsg = "Certificate does not include a private key and cannot be used to code sign";
 
-						// Prompt for a key pair
-						if (MessageBox.Show("Next, please choose the key pair that you made when generating the certificate request.",
-							Config.AppDisplayName,
-							MessageBoxButtons.OK,
-							MessageBoxIcon.Information) == DialogResult.OK)
+						// If the certificate doesn't have a private key pair, ask the user to provide one
+						if (ShowPrompt && !Cert.HasPrivateKey)
 						{
-							string KeyFilename;
-							if (ShowOpenFileDialog(KeysFilter, "Choose the key pair that belongs with the signing certificate", "", "", ref ChoosingFilesToInstallDirectory, out KeyFilename))
-							{
-								Cert = CryptoAdapter.CombineKeyAndCert(CertificateFilename, KeyFilename);
+							string ErrorMsg = "Certificate does not include a private key and cannot be used to code sign";
 
-								if (Cert.HasPrivateKey)
+							// Prompt for a key pair
+							if (MessageBox.Show("Next, please choose the key pair that you made when generating the certificate request.",
+								Config.AppDisplayName,
+								MessageBoxButtons.OK,
+								MessageBoxIcon.Information) == DialogResult.OK)
+							{
+								string KeyFilename;
+								if (ShowOpenFileDialog(KeysFilter, "Choose the key pair that belongs with the signing certificate", "", "", ref ChoosingFilesToInstallDirectory, out KeyFilename))
 								{
-									ErrorMsg = null;
+									Cert = CryptoAdapter.CombineKeyAndCert(CertificateFilename, KeyFilename);
+										
+									if (Cert.HasPrivateKey)
+									{
+										ErrorMsg = null;
+									}
 								}
 							}
-						}
 
-						if (ErrorMsg != null)
-						{
-							throw new Exception(ErrorMsg);
+							if (ErrorMsg != null)
+							{
+								throw new Exception(ErrorMsg);
+							}
 						}
-					}
 
-					// Add the certificate to the store
-					X509Store Store = new X509Store();
-					Store.Open(OpenFlags.ReadWrite);
-					Store.Add(Cert);
-					Store.Close();
+						// Add the certificate to the store
+						X509Store Store = new X509Store();
+						Store.Open(OpenFlags.ReadWrite);
+						Store.Add(Cert);
+						Store.Close();
+					}
 				}
 			}
 			catch (Exception ex)
@@ -278,7 +307,7 @@ namespace iPhonePackager
 
 		private void ImportCertificateButton_Click(object sender, EventArgs e)
 		{
-			TryInstallingCertificate_PromptForKey();
+			TryInstallingCertificate_PromptForKey(null);
 			UpdateStatus();
 		}
 
diff --git a/Engine/Source/Programs/IOS/iPhonePackager/iPhonePackager.cs b/Engine/Source/Programs/IOS/iPhonePackager/iPhonePackager.cs
--- a/Engine/Source/Programs/IOS/iPhonePackager/iPhonePackager.cs
+++ b/Engine/Source/Programs/IOS/iPhonePackager/iPhonePackager.cs
@@ -224,108 +224,79 @@ namespace iPhonePackager
 							case "-cookonthefly":
 								Config.bCookOnTheFly = true;
 								break;
-							case "-provision":
-								Config.bProvision = true;
-								break;
-							case "-certificate":
-								Config.bCert = true;
-								break;
 						}
 
 						// get the stage dir path
-						if (Arg == "-stagedir")
-						{
+						if (Arg == "-stagedir") {
 							// make sure there's at least one more arg
-							if (Arguments.Length > ArgIndex + 1)
-							{
-								Config.RepackageStagingDirectory = Arguments[++ArgIndex];
-							}
-							else
-							{
+							if (Arguments.Length > ArgIndex + 1) {
+								Config.RepackageStagingDirectory = Arguments [++ArgIndex];
+							} else {
 								return false;
 							}
-						}
-						else if (Arg == "-config")
-						{
+						} else if (Arg == "-config") {
 							// make sure there's at least one more arg
-							if (Arguments.Length > ArgIndex + 1)
-							{
-								GameConfiguration = Arguments[++ArgIndex];
-							}
-							else
-							{
+							if (Arguments.Length > ArgIndex + 1) {
+								GameConfiguration = Arguments [++ArgIndex];
+							} else {
 								return false;
 							}
 						}
 						// append a name to the bungle identifier and display name
-						else if (Arg == "-bundlename")
-						{
-							if (Arguments.Length > ArgIndex + 1)
-							{
-								Config.OverrideBundleName = Arguments[++ArgIndex];
-							}
-							else
-							{
+						else if (Arg == "-bundlename") {
+							if (Arguments.Length > ArgIndex + 1) {
+								Config.OverrideBundleName = Arguments [++ArgIndex];
+							} else {
 								return false;
 							}
-						}
-						else if (Arg == "-mac")
-						{
+						} else if (Arg == "-mac") {
 							// make sure there's at least one more arg
-							if (Arguments.Length > ArgIndex + 1)
-							{
-								Config.OverrideMacName = Arguments[++ArgIndex];
-							}
-							else
-							{
+							if (Arguments.Length > ArgIndex + 1) {
+								Config.OverrideMacName = Arguments [++ArgIndex];
+							} else {
 								return false;
 							}
-						}
-						else if (Arg == "-architecture" || Arg == "-arch")
-						{
+						} else if (Arg == "-architecture" || Arg == "-arch") {
 							// make sure there's at least one more arg
-							if (Arguments.Length > ArgIndex + 1)
-							{
-								Architecture = "-" + Arguments[++ArgIndex];
-							}
-							else
-							{
+							if (Arguments.Length > ArgIndex + 1) {
+								Architecture = "-" + Arguments [++ArgIndex];
+							} else {
 								return false;
 							}
-						}
-						else if (Arg == "-project")
-						{
+						} else if (Arg == "-project") {
 							// make sure there's at least one more arg
-							if (Arguments.Length > ArgIndex + 1)
-							{
-								Config.ProjectFile = Arguments[++ArgIndex];
-							}
-							else
-							{
+							if (Arguments.Length > ArgIndex + 1) {
+								Config.ProjectFile = Arguments [++ArgIndex];
+							} else {
 								return false;
 							}
-						}
-						else if (Arg == "-device")
-						{
+						} else if (Arg == "-device") {
 							// make sure there's at least one more arg
-							if (Arguments.Length > ArgIndex + 1)
-							{
-								Config.DeviceId = Arguments[++ArgIndex];
+							if (Arguments.Length > ArgIndex + 1) {
+								Config.DeviceId = Arguments [++ArgIndex];
+							} else {
+								return false;
 							}
-							else
-							{
+						} else if (Arg == "-additionalcommandline") {
+							// make sure there's at least one more arg
+							if (Arguments.Length > ArgIndex + 1) {
+								AdditionalCommandline = Arguments [++ArgIndex];
+							} else {
 								return false;
 							}
-						}
-						else if (Arg == "-additionalcommandline")
-						{
+						} else if (Arg == "-provision") {
 							// make sure there's at least one more arg
-							if (Arguments.Length > ArgIndex + 1)
-							{
-								AdditionalCommandline = Arguments[++ArgIndex];
+							if (Arguments.Length > ArgIndex + 1) {
+								Config.Provision = Arguments [++ArgIndex];
+								Config.bProvision = true;
+							} else {
+								return false;
 							}
-							else
-							{
+						} else if (Arg == "-certificate") {
+							if (Arguments.Length > ArgIndex + 1) {
+							Config.Certificate = Arguments [++ArgIndex];
+								Config.bCert = true;
+							} else {
 								return false;
 							}
 						}
@@ -715,11 +686,11 @@ namespace iPhonePackager
 						GameName = "";
 						if (Config.bProvision)
 						{
-							ToolsHub.TryInstallingMobileProvision();
+							ToolsHub.TryInstallingMobileProvision(Config.Provision, false);
 						}
 						if (Config.bCert)
 						{
-							ToolsHub.TryInstallingCertificate_PromptForKey();
+							ToolsHub.TryInstallingCertificate_PromptForKey(Config.Certificate, false);
 						}
 						break;
 
