UE4 - Metal - Backend changes in prep for compute shaders

diff --git a/Engine/Source/Developer/iOS/MetalShaderFormat/Private/MetalBackend.cpp b/Engine/Source/Developer/iOS/MetalShaderFormat/Private/MetalBackend.cpp
--- a/Engine/Source/Developer/iOS/MetalShaderFormat/Private/MetalBackend.cpp
+++ b/Engine/Source/Developer/iOS/MetalShaderFormat/Private/MetalBackend.cpp
@@ -340,13 +340,6 @@ protected:
 	exec_list sampler_variables;
 	exec_list image_variables;
 
-	/** Data tied globally to the shader via attributes */
-	int wg_size_x;
-	int wg_size_y;
-	int wg_size_z;
-
-	glsl_tessellation_info tessellation;
-
 	/** Track global instructions. */
 	struct global_ir : public exec_node
 	{
@@ -367,7 +360,7 @@ protected:
 	exec_list used_md_arrays;
 
 	// Code generation flags
-	bool bIsVS;
+	_mesa_glsl_parser_targets Frequency;
 
 	FBuffers& Buffers;
 
@@ -516,6 +509,10 @@ protected:
 			check(t->inner_type);
 			print_base_type(t->inner_type);
 		}
+		else if (t->base_type == GLSL_TYPE_IMAGE)
+		{
+			// Do nothing...
+		}
 		else
 		{
 			check(t->HlslName);
@@ -617,10 +614,6 @@ protected:
 		}
 		else
 		{
-			int layout_bits =
-				(var->origin_upper_left ? 0x1 : 0) |
-				(var->pixel_center_integer ? 0x2 : 0);
-
 			if (scope_depth == 0 &&
 			   ((var->mode == ir_var_in) || (var->mode == ir_var_out)) && 
 			   var->is_interface_block)
@@ -629,7 +622,23 @@ protected:
 			}
 			else if (var->type->is_image())
 			{
-				check(0);
+				auto* PtrType = var->type->is_array() ? var->type->element_type() : var->type;
+				check(!PtrType->is_array() && PtrType->inner_type);
+
+				// Buffer
+				int BufferIndex = Buffers.GetIndex(var);
+				check(BufferIndex >= 0);
+				ralloc_asprintf_append(
+					buffer,
+					"device "
+					);
+				print_type_pre(PtrType->inner_type);
+				ralloc_asprintf_append(buffer, " *%s", unique_name(var));
+				print_type_post(PtrType->inner_type);
+				ralloc_asprintf_append(
+					buffer,
+					" [[ buffer(%d) ]]", BufferIndex
+					);
 			}
 			else
 			{
@@ -776,6 +785,10 @@ protected:
 				}
 				else
 				{
+					if (var->mode == ir_var_shared)
+					{
+						ralloc_asprintf_append(buffer, "threadgroup ");
+					}
 					print_type_pre(var->type);
 					ralloc_asprintf_append(buffer, " %s", unique_name(var));
 					print_type_post(var->type);
@@ -840,16 +853,6 @@ protected:
 			indentation--;
 		}
 
-		//grab the global attributes
-		if (sig->is_main)
-		{
-			wg_size_x = sig->wg_size_x;
-			wg_size_y = sig->wg_size_y;
-			wg_size_z = sig->wg_size_z;
-
-			tessellation = sig->tessellation;
-		}
-
 		indentation++;
 		foreach_iter(exec_list_iterator, iter, sig->body)
 		{
@@ -875,7 +878,21 @@ protected:
 				indent();
 				if (sig->is_main)
 				{
-					ralloc_asprintf_append(buffer, "%s ", bIsVS ? "vertex" : "fragment");
+					switch (Frequency)
+					{
+					case vertex_shader:
+						ralloc_asprintf_append(buffer, "vertex ");
+						break;
+					case fragment_shader:
+						ralloc_asprintf_append(buffer, "fragment ");
+						break;
+					case compute_shader:
+						ralloc_asprintf_append(buffer, "kernel ");
+						break;
+					default:
+						check(0);
+						break;
+					}
 				}
 
 				sig->accept(this);
@@ -1170,7 +1187,7 @@ protected:
 		};
 		const char* int_cast[] =
 		{
-			"int", "ivec2", "ivec3", "ivec4"
+			"int", "int2", "int3", "int4"
 		};
 		const int dst_elements = deref->type->vector_elements;
 		const int src_elements = (src) ? src->type->vector_elements : 1;
@@ -1182,21 +1199,19 @@ protected:
 		{
 			if ( src == NULL )
 			{
-				ralloc_asprintf_append( buffer, "imageLoad( " );
 				deref->image->accept(this);
-				ralloc_asprintf_append(buffer, ", ");
+				ralloc_asprintf_append( buffer, "[" );
 				deref->image_index->accept(this);
-				ralloc_asprintf_append(buffer, ").%s", swizzle[dst_elements-1]);
+				ralloc_asprintf_append(buffer, "]"/*.%s, swizzle[dst_elements - 1]*/);
 			}
 			else
 			{
-				ralloc_asprintf_append( buffer, "imageStore( " );
 				deref->image->accept(this);
-				ralloc_asprintf_append(buffer, ", ");
+				ralloc_asprintf_append( buffer, "[" );
 				deref->image_index->accept(this);
-				ralloc_asprintf_append(buffer, ", ");
+				ralloc_asprintf_append( buffer, "] = " );
 				src->accept(this);
-				ralloc_asprintf_append(buffer, ".%s)", expand[src_elements-1]);
+				ralloc_asprintf_append(buffer, ""/*".%s", expand[src_elements - 1]*/);
 			}
 		}
 		else if ( deref->op == ir_image_dimensions)
@@ -1906,6 +1921,8 @@ protected:
 			if (hash_table_find(used_uniform_blocks, block->name))
 			{
 				const char* block_name = block->name;
+				check(0);
+/*
 				if (state->has_packed_uniforms)
 				{
 					block_name = ralloc_asprintf(mem_ctx, "%sb%u",
@@ -1925,7 +1942,7 @@ protected:
 					ralloc_asprintf_append(buffer, ";\n");
 				}
 				ralloc_asprintf_append(buffer, "};\n\n");
-
+*/
 				num_used_blocks++;
 			}
 		}
@@ -2350,30 +2367,12 @@ protected:
 		}
 	}
 
-	/**
-	 * Print the layout directives for this shader.
-	 */
-	void print_layout(_mesa_glsl_parse_state *state)
-	{
-		if (state->target == compute_shader )
-		{
-			ralloc_asprintf_append(buffer, "layout( local_size_x = %d, "
-				"local_size_y = %d, local_size_z = %d ) in;\n", wg_size_x,
-				wg_size_y, wg_size_z );
-		}
-
-		if(state->target == tessellation_evaluation_shader || state->target == tessellation_control_shader)
-		{
-			check(0);
-		}
-	}
-
 public:
 
 	/** Constructor. */
-	FGenerateMetalVisitor(_mesa_glsl_parse_state* InParseState, bool bInIsVS, FBuffers& InBuffers)
+	FGenerateMetalVisitor(_mesa_glsl_parse_state* InParseState, _mesa_glsl_parser_targets InFrequency, FBuffers& InBuffers)
 		: ParseState(InParseState)
-		, bIsVS(bInIsVS)
+		, Frequency(InFrequency)
 		, Buffers(InBuffers)
 		, buffer(0)
 		, indentation(0)
@@ -2428,15 +2427,10 @@ public:
 		print_signature(ParseState);
 		buffer = 0;
 
-		char* layout = ralloc_asprintf(mem_ctx, "");
-		buffer = &layout;
-		print_layout(ParseState);
-		buffer = 0;
 		char* full_buffer = ralloc_asprintf(
 			ParseState,
-			"// Compiled by HLSLCC\n%s\n#include <metal_stdlib>\nusing namespace metal;\n\n%s%s%s",
+			"// Compiled by HLSLCC\n%s\n#include <metal_stdlib>\nusing namespace metal;\n\n%s%s",
 			signature,
-			layout,
 			decl_buffer,
 			code_buffer
 			);
@@ -2450,7 +2444,7 @@ char* FMetalCodeBackend::GenerateCode(exec_list* ir, _mesa_glsl_parse_state* sta
 {
 	// We'll need this Buffers info for the [[buffer()]] index
 	FBuffers Buffers;
-	FGenerateMetalVisitor visitor(state, (state->target == vertex_shader), Buffers);
+	FGenerateMetalVisitor visitor(state, state->target, Buffers);
 
 	// At this point, all inputs and outputs are global uniforms, no structures.
 
diff --git a/Engine/Source/Developer/iOS/MetalShaderFormat/Private/MetalUtils.cpp b/Engine/Source/Developer/iOS/MetalShaderFormat/Private/MetalUtils.cpp
--- a/Engine/Source/Developer/iOS/MetalShaderFormat/Private/MetalUtils.cpp
+++ b/Engine/Source/Developer/iOS/MetalShaderFormat/Private/MetalUtils.cpp
@@ -245,12 +245,11 @@ namespace MetalUtils
 	/** Compute shader system values. */
 	static FSystemValue ComputeSystemValueTable[] =
 	{
-/*
-		{"SV_DispatchThreadID", glsl_type::uvec3_type, "gl_GlobalInvocationID", ir_var_in, false, false, false, false},
-		{"SV_GroupID", glsl_type::uvec3_type, "gl_WorkGroupID", ir_var_in, false, false, false, false},
-		{"SV_GroupIndex", glsl_type::uint_type, "gl_LocalInvocationIndex", ir_var_in, false, false, false, false},
-		{"SV_GroupThreadID", glsl_type::uvec3_type, "gl_LocalInvocationID", ir_var_in, false, false, false, false},
-*/
+		// D3D,					type,					GL,						param,		Metal
+		{"SV_DispatchThreadID",	glsl_type::uvec3_type,	"GlobalInvocationID",	ir_var_in, "[[ thread_position_in_grid ]]"},
+		{"SV_GroupID",			glsl_type::uvec3_type,	"WorkGroupID",			ir_var_in, "[[ threadgroup_position_in_grid ]]"},
+		{"SV_GroupIndex",		glsl_type::uint_type,	"LocalInvocationIndex",	ir_var_in, "[[ thread_index_in_threadgroup ]]"},
+		{"SV_GroupThreadID",	glsl_type::uvec3_type,	"LocalInvocationID",	ir_var_in, "[[ thread_position_in_threadgroup ]]"},
 		{NULL, NULL, NULL, ir_var_auto, nullptr}
 	};
 
@@ -1608,8 +1607,6 @@ void FMetalCodeBackend::PackInputsAndOutputs(exec_list* Instructions, _mesa_glsl
 	exec_list PostCallInstructions;
 	ParseState->symbols->push_scope();
 
-	// 
-
 	// Set of variables packed into a struct
 	std::set<ir_variable*> VSStageInVariables;
 	std::set<ir_variable*> PSStageInVariables;
@@ -1628,6 +1625,7 @@ void FMetalCodeBackend::PackInputsAndOutputs(exec_list* Instructions, _mesa_glsl
 	// Extra arguments needed for input (VertexID, etc)
 	TIRVarList VSInputArguments;
 	TIRVarList PSInputArguments;
+	TIRVarList CSInputArguments;
 
 	if (Frequency == HSF_VertexShader)
 	{
@@ -1812,6 +1810,60 @@ void FMetalCodeBackend::PackInputsAndOutputs(exec_list* Instructions, _mesa_glsl
 			}
 		}
 	}
+	else if (Frequency == HSF_ComputeShader)
+	{
+		YYLTYPE loc = {0};
+
+		foreach_iter(exec_list_iterator, Iter, *Instructions)
+		{
+			ir_instruction* IR = (ir_instruction*)Iter.get();
+			auto* Variable = IR->as_variable();
+			if (Variable)
+			{
+				switch (Variable->mode)
+				{
+				case ir_var_out:
+					{
+						_mesa_glsl_error(&loc, ParseState, "Compute/Kernel shaders do not support out variables ('%s')!", Variable->name);
+						return;
+					}
+					break;
+
+				case ir_var_in:
+					{
+						TArray<glsl_struct_field> CSStageInMembers;
+						TIRVarSet CSStageInVariables;
+						if (!ProcessStageInVariables(ParseState, Frequency, Variable, CSStageInMembers, CSStageInVariables, nullptr, CSInputArguments))
+						{
+							return;
+						}
+
+						if (CSStageInMembers.Num() != 0 || CSStageInVariables.size() != 0)
+						{
+							_mesa_glsl_error(&loc, ParseState, "Compute/Kernel shaders do not support out stage_in variables or vertex attributes ('%s')!", Variable->name);
+							return;
+						}
+					}
+					break;
+
+				case ir_var_shared:
+					{
+						// groupshared
+						Variable->remove();
+						DeclInstructions.push_head(Variable);
+					}
+					break;
+
+				default:
+					break;
+				}
+			}
+		}
+	}
+	else
+	{
+		check(0);
+	}
 
 	TIRVarList VarsToMoveToBody;
 	foreach_iter(exec_list_iterator, Iter, *Instructions)
