fix for a dangling reference to a process object
#ios

diff --git a/Engine/Source/Developer/iOS/IOSPlatformEditor/Private/IOSTargetSettingsCustomization.cpp b/Engine/Source/Developer/iOS/IOSPlatformEditor/Private/IOSTargetSettingsCustomization.cpp
--- a/Engine/Source/Developer/iOS/IOSPlatformEditor/Private/IOSTargetSettingsCustomization.cpp
+++ b/Engine/Source/Developer/iOS/IOSPlatformEditor/Private/IOSTargetSettingsCustomization.cpp
@@ -675,6 +675,7 @@ bool FIOSTargetSettingsCustomization::UpdateStatusDelegate(float DeltaTime)
 		return true;
 	}
 	int RetCode = IPPProcess->GetReturnCode();
+	IPPProcess = NULL;
 	ensure(RetCode == 0);
 	UpdateStatus();
 
