UE4 - Metal - remove packed from CB's (fixes array of float3 in cb's bug)

diff --git a/Engine/Source/Developer/iOS/MetalShaderFormat/Private/MetalBackend.cpp b/Engine/Source/Developer/iOS/MetalShaderFormat/Private/MetalBackend.cpp
--- a/Engine/Source/Developer/iOS/MetalShaderFormat/Private/MetalBackend.cpp
+++ b/Engine/Source/Developer/iOS/MetalShaderFormat/Private/MetalBackend.cpp
@@ -491,7 +491,13 @@ protected:
 	{
 		if (t->base_type == GLSL_TYPE_ARRAY)
 		{
+			bool bPrevPacked = bUsePacked;
+			if (t->element_type()->is_vector() && t->element_type()->vector_elements == 3)
+			{
+				bUsePacked = false;
+			}
 			print_base_type(t->fields.array);
+			bUsePacked = bPrevPacked;
 		}
 		else if (t->base_type == GLSL_TYPE_INPUTPATCH)
 		{
@@ -525,7 +531,7 @@ protected:
 		else
 		{
 			check(t->HlslName);
-			if (/*!bIsFunctionSig &&*/ bUsePacked && t->is_vector() && t->vector_elements < 4)
+			if (bUsePacked && t->is_vector() && t->vector_elements < 4)
 			{
 				ralloc_asprintf_append(buffer, "packed_%s", t->HlslName);
 			}
diff --git a/Engine/Source/Developer/iOS/MetalShaderFormat/Private/MetalShaderFormat.cpp b/Engine/Source/Developer/iOS/MetalShaderFormat/Private/MetalShaderFormat.cpp
--- a/Engine/Source/Developer/iOS/MetalShaderFormat/Private/MetalShaderFormat.cpp
+++ b/Engine/Source/Developer/iOS/MetalShaderFormat/Private/MetalShaderFormat.cpp
@@ -15,7 +15,7 @@ class FMetalShaderFormat : public IShaderFormat
 {
 	enum
 	{
-		HEADER_VERSION = 11,
+		HEADER_VERSION = 12,
 	};
 	struct FVersion
 	{
