make it so the architecture settings are not editable in binary release on PC and for content-only projects in binary release on Mac
UE-5084
#ios

diff --git a/Engine/Source/Developer/iOS/IOSPlatformEditor/IOSPlatformEditor.Build.cs b/Engine/Source/Developer/iOS/IOSPlatformEditor/IOSPlatformEditor.Build.cs
--- a/Engine/Source/Developer/iOS/IOSPlatformEditor/IOSPlatformEditor.Build.cs
+++ b/Engine/Source/Developer/iOS/IOSPlatformEditor/IOSPlatformEditor.Build.cs
@@ -15,6 +15,7 @@ public class IOSPlatformEditor : ModuleRules
 				"InputCore",
 				"DesktopPlatform",
 				"Engine",
+				"GameProjectGeneration",
 				"MainFrame",
 				"Slate",
 				"SlateCore",
diff --git a/Engine/Source/Developer/iOS/IOSPlatformEditor/Private/IOSTargetSettingsCustomization.cpp b/Engine/Source/Developer/iOS/IOSPlatformEditor/Private/IOSTargetSettingsCustomization.cpp
--- a/Engine/Source/Developer/iOS/IOSPlatformEditor/Private/IOSTargetSettingsCustomization.cpp
+++ b/Engine/Source/Developer/iOS/IOSPlatformEditor/Private/IOSTargetSettingsCustomization.cpp
@@ -18,6 +18,7 @@
 #include "SNotificationList.h"
 #include "NotificationManager.h"
 #include "TargetPlatform.h"
+#include "GameProjectGenerationModule.h"
 
 #define LOCTEXT_NAMESPACE "IOSTargetSettings"
 
@@ -97,7 +98,7 @@ void FIOSTargetSettingsCustomization::BuildPListSection(IDetailLayoutBuilder& De
 	IDetailCategoryBuilder& RenderCategory = DetailLayout.EditCategory(TEXT("Rendering"));
 	IDetailCategoryBuilder& OSInfoCategory = DetailLayout.EditCategory(TEXT("OS Info"));
 	IDetailCategoryBuilder& DeviceCategory = DetailLayout.EditCategory(TEXT("Devices"));
-	IDetailCategoryBuilder& CookCategory = DetailLayout.EditCategory(TEXT("Cook Settings"));
+	IDetailCategoryBuilder& BuildCategory = DetailLayout.EditCategory(TEXT("Build"));
 
 	TSharedRef<SPlatformSetupMessage> PlatformSetupMessage = SNew(SPlatformSetupMessage, GameInfoPath)
 		.PlatformName(LOCTEXT("iOSPlatformName", "iOS"))
@@ -255,6 +256,28 @@ void FIOSTargetSettingsCustomization::BuildPListSection(IDetailLayoutBuilder& De
 
 	// Show properties that are gated by the plist being present and writable
 	FSimpleDelegate PlistModifiedDelegate = FSimpleDelegate::CreateRaw(this, &FIOSTargetSettingsCustomization::OnPlistPropertyModified);
+	FGameProjectGenerationModule& GameProjectModule = FModuleManager::LoadModuleChecked<FGameProjectGenerationModule>(TEXT("GameProjectGeneration"));
+	bool bHasCode = GameProjectModule.Get().ProjectHasCodeFiles();
+
+#if PLATFORM_MAC
+#define SETUP_NONROCKET_PROP(PropName, Category, Tip, DisabledTip) \
+	{ \
+	TSharedRef<IPropertyHandle> PropertyHandle = DetailLayout.GetProperty(GET_MEMBER_NAME_CHECKED(UIOSRuntimeSettings, PropName)); \
+	Category.AddProperty(PropertyHandle) \
+		.EditCondition(SetupForPlatformAttribute, NULL) \
+		.IsEnabled(!FRocketSupport::IsRocket() || bHasCode) \
+		.ToolTip((!FRocketSupport::IsRocket() || bHasCode) ? Tip : DisabledTip); \
+	}
+#else
+#define SETUP_NONROCKET_PROP(PropName, Category, Tip, DisabledTip) \
+	{ \
+		TSharedRef<IPropertyHandle> PropertyHandle = DetailLayout.GetProperty(GET_MEMBER_NAME_CHECKED(UIOSRuntimeSettings, PropName)); \
+		Category.AddProperty(PropertyHandle) \
+			.EditCondition(SetupForPlatformAttribute, NULL) \
+			.IsEnabled(!FRocketSupport::IsRocket()) \
+			.ToolTip(!FRocketSupport::IsRocket() ? Tip : DisabledTip); \
+	}
+#endif
 
 #define SETUP_PLIST_PROP(PropName, Category, Tip) \
 	{ \
@@ -282,7 +305,16 @@ void FIOSTargetSettingsCustomization::BuildPListSection(IDetailLayoutBuilder& De
 
 	SETUP_PLIST_PROP(MinimumiOSVersion, OSInfoCategory, TEXT("WMinimum iOS version this game supports"));
 
+	FString DisabledTip = TEXT("This requires GitHub source.");
+	SETUP_NONROCKET_PROP(bDevForArmV7, BuildCategory, TEXT("Enable ArmV7 support? (this will be used if all type are unchecked)"), DisabledTip);
+	SETUP_NONROCKET_PROP(bDevForArm64, BuildCategory, TEXT("Enable Arm64 support?"), DisabledTip);
+	SETUP_NONROCKET_PROP(bDevForArmV7S, BuildCategory, TEXT("Enable ArmV7s support?"), DisabledTip);
+	SETUP_NONROCKET_PROP(bShipForArmV7, BuildCategory, TEXT("Enable ArmV7 support? (this will be used if all type are unchecked)"), DisabledTip);
+	SETUP_NONROCKET_PROP(bShipForArm64, BuildCategory, TEXT("Enable Arm64 support?"), DisabledTip);
+	SETUP_NONROCKET_PROP(bShipForArmV7S, BuildCategory, TEXT("Enable ArmV7s support?"), DisabledTip);
+
 #undef SETUP_PLIST_PROP
+#undef SETUP_NONROCKET_PROP
 }
 
 void FIOSTargetSettingsCustomization::BuildIconSection(IDetailLayoutBuilder& DetailLayout)
diff --git a/Engine/Source/Runtime/IOS/IOSRuntimeSettings/Classes/IOSRuntimeSettings.h b/Engine/Source/Runtime/IOS/IOSRuntimeSettings/Classes/IOSRuntimeSettings.h
--- a/Engine/Source/Runtime/IOS/IOSRuntimeSettings/Classes/IOSRuntimeSettings.h
+++ b/Engine/Source/Runtime/IOS/IOSRuntimeSettings/Classes/IOSRuntimeSettings.h
@@ -55,27 +55,27 @@ public:
 	UPROPERTY(GlobalConfig, EditAnywhere, Category = Rendering)
 	bool bSupportsOpenGLES2;
 
-	// Enable ArmV7 support? (this will be used if all type are unchecked)
+	// Enable ArmV7 support? (this will be used if all type are unchecked) [CURRENTLY AVAILABLE ON MAC ONLY OR FULL SOURCE BUILD]
 	UPROPERTY(GlobalConfig, EditAnywhere, Category = Build, meta = (DisplayName = "Support armv7 in Development"))
 	bool bDevForArmV7;
 
-	// Enable Arm64 support?
+	// Enable Arm64 support? [CURRENTLY AVAILABLE ON MAC ONLY OR FULL SOURCE BUILD]
 	UPROPERTY(GlobalConfig, EditAnywhere, Category = Build, meta = (DisplayName = "Support arm64 in Development"))
 	bool bDevForArm64;
 
-	// Enable x86 support? [CURRENTLY FOR FULL SOURCE GAMES ONLY]
+	// Enable ArmV7s support? [CURRENTLY AVAILABLE ON MAC ONLY OR FULL SOURCE BUILD]
 	UPROPERTY(GlobalConfig, EditAnywhere, Category = Build, meta = (DisplayName = "Support armv7s in Development"))
 	bool bDevForArmV7S;
 
-	// Enable ArmV7 support? (this will be used if all type are unchecked)
+	// Enable ArmV7 support for shipping build? (this will be used if all type are unchecked) [CURRENTLY AVAILABLE ON MAC ONLY OR FULL SOURCE BUILD]
 	UPROPERTY(GlobalConfig, EditAnywhere, Category = Build, meta = (DisplayName = "Support armv7 in Shipping"))
 	bool bShipForArmV7;
 
-	// Enable Arm64 support?
+	// Enable Arm64 support for shipping build? [CURRENTLY AVAILABLE ON MAC ONLY OR FULL SOURCE BUILD]
 	UPROPERTY(GlobalConfig, EditAnywhere, Category = Build, meta = (DisplayName = "Support arm64 in Shipping"))
 	bool bShipForArm64;
 
-	// Enable x86 support? [CURRENTLY FOR FULL SOURCE GAMES ONLY]
+	// Enable ArmV7s support for shipping build? [CURRENTLY AVAILABLE ON MAC ONLY OR FULL SOURCE BUILD]
 	UPROPERTY(GlobalConfig, EditAnywhere, Category = Build, meta = (DisplayName = "Support armv7s in Shipping"))
 	bool bShipForArmV7S;
 
