- Changed Metal clamp to precise version to remove NaNs

diff --git a/Engine/Source/Developer/iOS/MetalShaderFormat/Private/MetalBackend.cpp b/Engine/Source/Developer/iOS/MetalShaderFormat/Private/MetalBackend.cpp
--- a/Engine/Source/Developer/iOS/MetalShaderFormat/Private/MetalBackend.cpp
+++ b/Engine/Source/Developer/iOS/MetalShaderFormat/Private/MetalBackend.cpp
@@ -184,7 +184,7 @@ static const char * const MetalExpressionTable[ir_opcode_count][4] =
 
 	{ "mix(", ",", ",", ")" }, // ir_ternop_lerp,
 	{ "smoothstep(", ",", ",", ")" }, // ir_ternop_smoothstep,
-	{ "clamp(", ",", ",", ")" }, // ir_ternop_clamp,
+	{ "precise::clamp(", ",", ",", ")" }, // ir_ternop_clamp,
 
 	{ "ERROR_QUADOP_VECTOR(", ",", ")" }, // ir_quadop_vector,
 };
