disable architecture selection for all binary release projects on both PC and Mac
UE-5733
#ios

diff --git a/Engine/Source/Developer/iOS/IOSPlatformEditor/Private/IOSTargetSettingsCustomization.cpp b/Engine/Source/Developer/iOS/IOSPlatformEditor/Private/IOSTargetSettingsCustomization.cpp
--- a/Engine/Source/Developer/iOS/IOSPlatformEditor/Private/IOSTargetSettingsCustomization.cpp
+++ b/Engine/Source/Developer/iOS/IOSPlatformEditor/Private/IOSTargetSettingsCustomization.cpp
@@ -259,16 +259,6 @@ void FIOSTargetSettingsCustomization::BuildPListSection(IDetailLayoutBuilder& De
 	FGameProjectGenerationModule& GameProjectModule = FModuleManager::LoadModuleChecked<FGameProjectGenerationModule>(TEXT("GameProjectGeneration"));
 	bool bHasCode = GameProjectModule.Get().ProjectHasCodeFiles();
 
-#if PLATFORM_MAC
-#define SETUP_NONROCKET_PROP(PropName, Category, Tip, DisabledTip) \
-	{ \
-	TSharedRef<IPropertyHandle> PropertyHandle = DetailLayout.GetProperty(GET_MEMBER_NAME_CHECKED(UIOSRuntimeSettings, PropName)); \
-	Category.AddProperty(PropertyHandle) \
-		.EditCondition(SetupForPlatformAttribute, NULL) \
-		.IsEnabled(!FRocketSupport::IsRocket() || bHasCode) \
-		.ToolTip((!FRocketSupport::IsRocket() || bHasCode) ? Tip : DisabledTip); \
-	}
-#else
 #define SETUP_NONROCKET_PROP(PropName, Category, Tip, DisabledTip) \
 	{ \
 		TSharedRef<IPropertyHandle> PropertyHandle = DetailLayout.GetProperty(GET_MEMBER_NAME_CHECKED(UIOSRuntimeSettings, PropName)); \
@@ -277,7 +267,6 @@ void FIOSTargetSettingsCustomization::BuildPListSection(IDetailLayoutBuilder& De
 			.IsEnabled(!FRocketSupport::IsRocket()) \
 			.ToolTip(!FRocketSupport::IsRocket() ? Tip : DisabledTip); \
 	}
-#endif
 
 #define SETUP_PLIST_PROP(PropName, Category, Tip) \
 	{ \
