- Cleaned up the IOS texture format decision code
- Cleaned up the code to decide what shader formats to use
- Added separate IOSRuntimeSetting for bSupportsMetalMRT, and check it when using -metalmrt flag

diff --git a/Engine/Config/BaseEngine.ini b/Engine/Config/BaseEngine.ini
--- a/Engine/Config/BaseEngine.ini
+++ b/Engine/Config/BaseEngine.ini
@@ -2353,6 +2353,7 @@ bSupportsUpsideDownOrientation=False
 bSupportsLandscapeLeftOrientation=True
 bSupportsLandscapeRightOrientation=True
 bSupportsMetal=False
+bSupportsMetal=True
 bSupportsOpenGLES2=True
 
 ; These are the defaults for Android settings, and they need to be in the .ini since UBT reads the .ini settings, without instantiating the class
diff --git a/Engine/Source/Developer/TargetPlatform/Public/Common/TargetPlatformBase.h b/Engine/Source/Developer/TargetPlatform/Public/Common/TargetPlatformBase.h
--- a/Engine/Source/Developer/TargetPlatform/Public/Common/TargetPlatformBase.h
+++ b/Engine/Source/Developer/TargetPlatform/Public/Common/TargetPlatformBase.h
@@ -161,7 +161,10 @@ public:
 	}
 #else //TEXTURERESOURCE_H_INCLUDED
 
-	FName GetDefaultTextureFormatName( const UTexture* Texture, const FConfigFile& EngineSettings ) const;
+	FName GetDefaultTextureFormatName( const UTexture* Texture, const FConfigFile& EngineSettings ) const
+	{
+		return NAME_None;
+	}
 
 #endif //TEXTURERESOURCE_H_INCLUDED
 #endif //WITH_ENGINE
diff --git a/Engine/Source/Developer/iOS/IOSTargetPlatform/Private/IOSTargetPlatform.cpp b/Engine/Source/Developer/iOS/IOSTargetPlatform/Private/IOSTargetPlatform.cpp
--- a/Engine/Source/Developer/iOS/IOSTargetPlatform/Private/IOSTargetPlatform.cpp
+++ b/Engine/Source/Developer/iOS/IOSTargetPlatform/Private/IOSTargetPlatform.cpp
@@ -290,28 +290,49 @@ bool FIOSTargetPlatform::HandleTicker(float DeltaTime )
 
 #if WITH_ENGINE
 
+static bool SupportsES2()
+{
+	// default to supporting ES2
+	bool bSupportsOpenGLES2 = true;
+	GConfig->GetBool(TEXT("/Script/IOSRuntimeSettings.IOSRuntimeSettings"), TEXT("bSupportsOpenGLES2"), bSupportsOpenGLES2, GEngineIni);
+	return bSupportsOpenGLES2;
+}
+
+static bool SupportsMetal()
+{
+	// default to NOT supporting metal
+	bool bSupportsMetal = false;
+	GConfig->GetBool(TEXT("/Script/IOSRuntimeSettings.IOSRuntimeSettings"), TEXT("bSupportsMetal"), bSupportsMetal, GEngineIni);
+	return bSupportsMetal;
+}
+
+static bool SupportsMetalMRT()
+{
+	// default to NOT supporting metal MRT
+	bool bSupportsMetalMRT = false;
+	GConfig->GetBool(TEXT("/Script/IOSRuntimeSettings.IOSRuntimeSettings"), TEXT("bSupportsMetalMRT"), bSupportsMetalMRT, GEngineIni);
+	return bSupportsMetalMRT;
+}
+
 void FIOSTargetPlatform::GetAllPossibleShaderFormats( TArray<FName>& OutFormats ) const
 {
 	static FName NAME_OPENGL_ES2_IOS(TEXT("GLSL_ES2_IOS"));
 	static FName NAME_SF_METAL(TEXT("SF_METAL"));
 	static FName NAME_SF_METAL_MRT(TEXT("SF_METAL_MRT"));
 
-	// default to supporting ES2
-	bool bSupportsOpenGLES2 = true;
-	GConfig->GetBool(TEXT("/Script/IOSRuntimeSettings.IOSRuntimeSettings"), TEXT("bSupportsOpenGLES2"), bSupportsOpenGLES2, GEngineIni);
-	if (bSupportsOpenGLES2)
+	if (SupportsES2())
 	{
 		OutFormats.AddUnique(NAME_OPENGL_ES2_IOS);
 	}
 
-	// default to NOT supporting metal
-	bool bSupportsMetal = false;
-	GConfig->GetBool(TEXT("/Script/IOSRuntimeSettings.IOSRuntimeSettings"), TEXT("bSupportsMetal"), bSupportsMetal, GEngineIni);
-	if (bSupportsMetal)
+	if (SupportsMetal())
 	{
 		OutFormats.AddUnique(NAME_SF_METAL);
-		// @todo metal mrt: Add a project setting to support MetalMRT mode
-		// OutFormats.AddUnique(NAME_SF_METAL_MRT);
+	}
+
+	if (SupportsMetalMRT())
+	{
+		OutFormats.AddUnique(NAME_SF_METAL_MRT);
 	}
 }
 
@@ -321,111 +342,41 @@ void FIOSTargetPlatform::GetAllTargetedShaderFormats( TArray<FName>& OutFormats
 	GetAllPossibleShaderFormats(OutFormats);
 }
 
-
 void FIOSTargetPlatform::GetTextureFormats( const UTexture* Texture, TArray<FName>& OutFormats ) const
 {
 	check(Texture);
 
-	FName TextureFormatName = NAME_None;
-
-	//// Supported texture format names.
-
-	// Compressed Texture Formats
-	static FName NamePVRTC2(TEXT("PVRTC2"));
-	static FName NamePVRTC4(TEXT("PVRTC4"));
-
-	// Same as PVRTC4, but derives Z from X and Y
-	static FName NamePVRTCN(TEXT("PVRTCN"));
-	static FName NameAutoPVRTC(TEXT("AutoPVRTC"));
-
-	// Uncompressed Texture Formats
-	static FName NameBGRA8(TEXT("BGRA8"));
-	static FName NameG8(TEXT("G8"));
-	static FName NameVU8(TEXT("VU8"));
-	static FName NameRGBA16F(TEXT("RGBA16F"));
-	bool bIsCubemap = Texture->IsA(UTextureCube::StaticClass());
-
-	bool bNoCompression = Texture->CompressionNone				// Code wants the texture uncompressed.
-		|| (Texture->LODGroup == TEXTUREGROUP_ColorLookupTable)	// Textures in certain LOD groups should remain uncompressed.
-		|| (Texture->LODGroup == TEXTUREGROUP_Bokeh)
-		|| (Texture->CompressionSettings == TC_EditorIcon)
-		|| (Texture->Source.GetSizeX() < 4)						// Don't compress textures smaller than the DXT block size.
-		|| (Texture->Source.GetSizeY() < 4)
-		|| (Texture->Source.GetSizeX() % 4 != 0)
-		|| (Texture->Source.GetSizeY() % 4 != 0);
+	// we remap some of the defaults
+	static FName FormatRemap[] = 
+	{
+		FName(TEXT("DXT1")), FName(TEXT("PVRTC2")),
+		FName(TEXT("DXT5")), FName(TEXT("PVRTC4")),
+		FName(TEXT("DXT5n")), FName(TEXT("PVRTCN")),
+		FName(TEXT("BC5")), FName(TEXT("PVRTCN")),
+		FName(TEXT("AutoDXT")), FName(TEXT("AutoPVRTC")),
+	};
 
-	ETextureSourceFormat SourceFormat = Texture->Source.GetFormat();
+	FName TextureFormatName = NAME_None;
 
-	// Determine the pixel format of the compressed texture.
-	if (bNoCompression)
-	{
-		if (Texture->HasHDRSource())
-		{
-			TextureFormatName = NameRGBA16F;
-		}
-		else if (SourceFormat == TSF_G8 || Texture->CompressionSettings == TC_Grayscale)
-		{
-			TextureFormatName = NameG8;
-		}
-		else if (Texture->LODGroup == TEXTUREGROUP_Shadowmap)
-		{
-			TextureFormatName = NameG8;
-		}
-		else
-		{
-			TextureFormatName = NameBGRA8;
-		}
-	}
-	else if (Texture->CompressionSettings == TC_HDR)
-	{
-		TextureFormatName = NameRGBA16F;
-	}
-	else if (Texture->CompressionSettings == TC_Normalmap)
-	{
-		TextureFormatName = NamePVRTCN;
-	}
-	else if (Texture->CompressionSettings == TC_Displacementmap)
-	{
-		TextureFormatName = NameG8;
-	}
-	else if (Texture->CompressionSettings == TC_VectorDisplacementmap)
-	{
-		TextureFormatName = NameBGRA8;
-	}
-	else if (Texture->CompressionSettings == TC_Grayscale)
-	{
-		TextureFormatName = NameG8;
-	}
-	else if (Texture->CompressionSettings == TC_Alpha)
-	{
-		TextureFormatName = NameG8;
-	}
-	else if (Texture->CompressionSettings == TC_DistanceFieldFont)
-	{
-		TextureFormatName = NameG8;
-	}
-	else if (Texture->bForcePVRTC4)
-	{
-		TextureFormatName = NamePVRTC4;
-	}
-	else if (Texture->CompressionNoAlpha)
-	{
-		TextureFormatName = NamePVRTC2;
-	}
-	else if (Texture->bDitherMipMapAlpha)
+	// forward rendering only needs one channel for shadow maps
+	if (Texture->LODGroup == TEXTUREGROUP_Shadowmap && !SupportsMetalMRT())
 	{
-		TextureFormatName = NamePVRTC4;
+		TextureFormatName = FName(TEXT("G8"));
 	}
-	else
+
+	// if we didn't assign anything specially, then use the defaults
+	if (TextureFormatName == NAME_None)
 	{
-		TextureFormatName = NameAutoPVRTC;
+		TextureFormatName = GetDefaultTextureFormatName(Texture, EngineSettings);
 	}
 
-	// Some PC GPUs don't support sRGB read from G8 textures (e.g. AMD DX10 cards on ShaderModel3.0)
-	// This solution requires 4x more memory but a lot of PC HW emulate the format anyway
-	if ((TextureFormatName == NameG8) && Texture->SRGB)
+	// perform any remapping away from defaults
+	for (int32 RemapIndex = 0; RemapIndex < ARRAY_COUNT(FormatRemap); RemapIndex += 2)
 	{
-		TextureFormatName = NameG8;
+		if (TextureFormatName == FormatRemap[RemapIndex])
+		{
+			TextureFormatName = FormatRemap[RemapIndex + 1];
+		}
 	}
 
 	OutFormats.Add(TextureFormatName);
diff --git a/Engine/Source/Developer/iOS/IOSTargetPlatform/Private/IOSTargetPlatformPrivatePCH.h b/Engine/Source/Developer/iOS/IOSTargetPlatform/Private/IOSTargetPlatformPrivatePCH.h
--- a/Engine/Source/Developer/iOS/IOSTargetPlatform/Private/IOSTargetPlatformPrivatePCH.h
+++ b/Engine/Source/Developer/iOS/IOSTargetPlatform/Private/IOSTargetPlatformPrivatePCH.h
@@ -16,15 +16,15 @@
 #include "LaunchDaemonMessages.h"
 #include "Messaging.h"
 #include "TargetDeviceServices.h"
-#include "TargetPlatform.h"
-#include "TargetPlatformBase.h"
 #include "IOSPlatformProperties.h"
 
 #if WITH_ENGINE
 	#include "Engine.h"
-	//#include "TextureCompressorModule.h"
 #endif
 
+#include "TargetPlatform.h"
+#include "TargetPlatformBase.h"
+
 
 /* Private includes
  *****************************************************************************/
diff --git a/Engine/Source/Runtime/Core/Private/IOS/IOSView.cpp b/Engine/Source/Runtime/Core/Private/IOS/IOSView.cpp
--- a/Engine/Source/Runtime/Core/Private/IOS/IOSView.cpp
+++ b/Engine/Source/Runtime/Core/Private/IOS/IOSView.cpp
@@ -26,7 +26,9 @@ id<MTLDevice> GMetalDevice = nil;
 #if HAS_METAL
 	// make sure the project setting has enabled Metal support (per-project user settings in the editor)
 	bool bSupportsMetal = false;
+	bool bSupportsMetalMRT = false;
 	GConfig->GetBool(TEXT("/Script/IOSRuntimeSettings.IOSRuntimeSettings"), TEXT("bSupportsMetal"), bSupportsMetal, GEngineIni);
+	GConfig->GetBool(TEXT("/Script/IOSRuntimeSettings.IOSRuntimeSettings"), TEXT("bSupportsMetalMRT"), bSupportsMetalMRT, GEngineIni);
 
 	// does commandline override?
 	bool bForceES2 = FParse::Param(FCommandLine::Get(), TEXT("ES2"));
diff --git a/Engine/Source/Runtime/IOS/IOSRuntimeSettings/Classes/IOSRuntimeSettings.h b/Engine/Source/Runtime/IOS/IOSRuntimeSettings/Classes/IOSRuntimeSettings.h
--- a/Engine/Source/Runtime/IOS/IOSRuntimeSettings/Classes/IOSRuntimeSettings.h
+++ b/Engine/Source/Runtime/IOS/IOSRuntimeSettings/Classes/IOSRuntimeSettings.h
@@ -48,9 +48,13 @@ public:
 	uint32 bEnableGameCenterSupport : 1;
 	
 	// Whether or not to add support for Metal API (requires IOS8 and A7 processors)
-	UPROPERTY(GlobalConfig, EditAnywhere, Category = Rendering)
+	UPROPERTY(GlobalConfig, EditAnywhere, Category = Rendering, meta = (DisplayName = "Supports Forward Rendering with Metal (A7 and up devices)"))
 	bool bSupportsMetal;
 
+	// Whether or not to add support for deferred rendering Metal API (requires IOS8 and A7 processors)
+	UPROPERTY(GlobalConfig, EditAnywhere, Category = Rendering, meta = (DisplayName = "Supports Deferred Rendering with Metal (A8 and up devices)"))
+	bool bSupportsMetalMRT;
+	
 	// Whether or not to add support for OpenGL ES2 (if this is false, then your game should specify minimum IOS8 version and use "metal" instead of "opengles-2" in UIRequiredDeviceCapabilities)
 	UPROPERTY(GlobalConfig, EditAnywhere, Category = Rendering)
 	bool bSupportsOpenGLES2;
diff --git a/Engine/Source/Runtime/IOS/IOSRuntimeSettings/Private/IOSRuntimeSettings.cpp b/Engine/Source/Runtime/IOS/IOSRuntimeSettings/Private/IOSRuntimeSettings.cpp
--- a/Engine/Source/Runtime/IOS/IOSRuntimeSettings/Private/IOSRuntimeSettings.cpp
+++ b/Engine/Source/Runtime/IOS/IOSRuntimeSettings/Private/IOSRuntimeSettings.cpp
@@ -37,7 +37,7 @@ void UIOSRuntimeSettings::PostEditChangeProperty(struct FPropertyChangedEvent& P
 	}
 
 	// Ensure that at least one API is supported
-	if (!bSupportsMetal && !bSupportsOpenGLES2)
+	if (!bSupportsMetal && !bSupportsOpenGLES2 && !bSupportsMetalMRT)
 	{
 		bSupportsOpenGLES2 = true;
 	}
diff --git a/Engine/Source/Runtime/IOS/MetalRHI/Private/MetalRHI.cpp b/Engine/Source/Runtime/IOS/MetalRHI/Private/MetalRHI.cpp
--- a/Engine/Source/Runtime/IOS/MetalRHI/Private/MetalRHI.cpp
+++ b/Engine/Source/Runtime/IOS/MetalRHI/Private/MetalRHI.cpp
@@ -46,8 +46,12 @@ FMetalDynamicRHI::FMetalDynamicRHI()
 	id<MTLDevice> Device = [IOSAppDelegate GetDelegate].IOSView->MetalDevice;
 	// A8 can use 256 bits of MRTs
 	bool bCanUseWideMRTs = [Device supportsFeatureSet:MTLFeatureSet_iOS_GPUFamily2_v1];
+
+	bool bProjectSupportsMRTs = false;
+	GConfig->GetBool(TEXT("/Script/IOSRuntimeSettings.IOSRuntimeSettings"), TEXT("bSupportsMetalMRT"), bProjectSupportsMRTs, GEngineIni);
+
 	// only allow GBuffers, etc on A8s (A7s are just not going to cut it)
-	if (bCanUseWideMRTs && FParse::Param(FCommandLine::Get(),TEXT("metalmrt")))
+	if (bProjectSupportsMRTs && bCanUseWideMRTs && FParse::Param(FCommandLine::Get(),TEXT("metalmrt")))
 	{
 		GMaxRHIFeatureLevel = ERHIFeatureLevel::SM4;
 		GMaxRHIShaderPlatform = SP_METAL_MRT;
