fix for incorrectly scaled images and buttons for mobile provision and certificate in IOS settings
UE-4918
#ios

diff --git a/Engine/Source/Developer/iOS/IOSPlatformEditor/Private/IOSTargetSettingsCustomization.cpp b/Engine/Source/Developer/iOS/IOSPlatformEditor/Private/IOSTargetSettingsCustomization.cpp
--- a/Engine/Source/Developer/iOS/IOSPlatformEditor/Private/IOSTargetSettingsCustomization.cpp
+++ b/Engine/Source/Developer/iOS/IOSPlatformEditor/Private/IOSTargetSettingsCustomization.cpp
@@ -148,21 +148,24 @@ void FIOSTargetSettingsCustomization::BuildPListSection(IDetailLayoutBuilder& De
 			]
 		]
 		.ValueContent()
+			.MinDesiredWidth(300.0f)
+			.MaxDesiredWidth(350.0f)
 		[
 			SNew(SHorizontalBox)
 			+ SHorizontalBox::Slot()
-			.Padding(FMargin(8, 0, 8, 0))
+			.Padding(FMargin(8,3,8,3))
 			.AutoWidth()
 			[
 				SNew(SImage)
 				.Image(this, &FIOSTargetSettingsCustomization::GetProvisionStatus)
 			]
 			+ SHorizontalBox::Slot()
-			.Padding(FMargin(0, 1, 0, 1))
-			.FillWidth(1.0f)
+			.Padding(FMargin(0, 0, 0, 0))
+			.AutoWidth()
 			[
 				SNew(SButton)
 				.HAlign(HAlign_Center)
+				.VAlign(VAlign_Center)
 				.OnClicked(this, &FIOSTargetSettingsCustomization::OnInstallProvisionClicked)
 				[
 					SNew(STextBlock)
@@ -185,27 +188,29 @@ void FIOSTargetSettingsCustomization::BuildPListSection(IDetailLayoutBuilder& De
 			]
 		]
 		.ValueContent()
+			.MinDesiredWidth(300.0f)
+			.MaxDesiredWidth(350.0f)
 			[
 				SNew(SHorizontalBox)
 				+ SHorizontalBox::Slot()
-				.Padding(FMargin(8, 0, 8, 0))
 				.AutoWidth()
+				.Padding(FMargin(8,3,8,3))
 				[
 					SNew(SImage)
 					.Image(this, &FIOSTargetSettingsCustomization::GetCertificateStatus)
 				]
 				+ SHorizontalBox::Slot()
-				.Padding(FMargin(0, 1, 0, 1))
-				.FillWidth(1.0f)
-				[
-					SNew(SButton)
-					.HAlign(HAlign_Center)
-					.OnClicked(this, &FIOSTargetSettingsCustomization::OnInstallCertificateClicked)
+					.AutoWidth()
 					[
-						SNew(STextBlock)
-						.Text(FText::FromString("Import Certificate"))
+						SNew(SButton)
+						.HAlign(HAlign_Center)
+						.VAlign(VAlign_Center)
+						.OnClicked(this, &FIOSTargetSettingsCustomization::OnInstallCertificateClicked)
+						[
+							SNew(STextBlock)
+							.Text(FText::FromString("Import Certificate"))
+						]
 					]
-				]
 			];
 
 	AppManifestCategory.AddCustomRow(TEXT("Warning"), false)
