UE4 - Undo Metal precise clamp

diff --git a/Engine/Source/Developer/iOS/MetalShaderFormat/Private/MetalBackend.cpp b/Engine/Source/Developer/iOS/MetalShaderFormat/Private/MetalBackend.cpp
--- a/Engine/Source/Developer/iOS/MetalShaderFormat/Private/MetalBackend.cpp
+++ b/Engine/Source/Developer/iOS/MetalShaderFormat/Private/MetalBackend.cpp
@@ -184,7 +184,7 @@ static const char * const MetalExpressionTable[ir_opcode_count][4] =
 
 	{ "mix(", ",", ",", ")" }, // ir_ternop_lerp,
 	{ "smoothstep(", ",", ",", ")" }, // ir_ternop_smoothstep,
-	{ "precise::clamp(", ",", ",", ")" }, // ir_ternop_clamp,
+	{ "clamp(", ",", ",", ")" }, // ir_ternop_clamp,
 
 	{ "ERROR_QUADOP_VECTOR(", ",", ")" }, // ir_quadop_vector,
 };
diff --git a/Engine/Source/Developer/iOS/MetalShaderFormat/Private/MetalShaderFormat.cpp b/Engine/Source/Developer/iOS/MetalShaderFormat/Private/MetalShaderFormat.cpp
--- a/Engine/Source/Developer/iOS/MetalShaderFormat/Private/MetalShaderFormat.cpp
+++ b/Engine/Source/Developer/iOS/MetalShaderFormat/Private/MetalShaderFormat.cpp
@@ -15,7 +15,7 @@ class FMetalShaderFormat : public IShaderFormat
 {
 	enum
 	{
-		HEADER_VERSION = 12,
+		HEADER_VERSION = 13,
 	};
 	struct FVersion
 	{
