UE4 - Metal - Use precise::clamp for float

diff --git a/Engine/Source/Developer/iOS/MetalShaderFormat/Private/MetalBackend.cpp b/Engine/Source/Developer/iOS/MetalShaderFormat/Private/MetalBackend.cpp
--- a/Engine/Source/Developer/iOS/MetalShaderFormat/Private/MetalBackend.cpp
+++ b/Engine/Source/Developer/iOS/MetalShaderFormat/Private/MetalBackend.cpp
@@ -990,6 +990,15 @@ protected:
 			ralloc_asprintf_append(buffer, "ERRROR_MulMatrix()");
 			check(0);
 		}
+		else if (op == ir_ternop_clamp && expr->type->base_type == GLSL_TYPE_FLOAT)
+		{
+			ralloc_asprintf_append(buffer, "precise::%s", MetalExpressionTable[op][0]);
+			for (int i = 0; i < numOps; ++i)
+			{
+				expr->operands[i]->accept(this);
+				ralloc_asprintf_append(buffer, MetalExpressionTable[op][i+1]);
+			}
+		}
 		else if (numOps < 4)
 		{
 			ralloc_asprintf_append(buffer, MetalExpressionTable[op][0]);
diff --git a/Engine/Source/Developer/iOS/MetalShaderFormat/Private/MetalShaderFormat.cpp b/Engine/Source/Developer/iOS/MetalShaderFormat/Private/MetalShaderFormat.cpp
--- a/Engine/Source/Developer/iOS/MetalShaderFormat/Private/MetalShaderFormat.cpp
+++ b/Engine/Source/Developer/iOS/MetalShaderFormat/Private/MetalShaderFormat.cpp
@@ -15,7 +15,7 @@ class FMetalShaderFormat : public IShaderFormat
 {
 	enum
 	{
-		HEADER_VERSION = 13,
+		HEADER_VERSION = 14,
 	};
 	struct FVersion
 	{
diff --git a/Engine/Source/Developer/iOS/MetalShaderFormat/Private/MetalUtils.cpp b/Engine/Source/Developer/iOS/MetalShaderFormat/Private/MetalUtils.cpp
--- a/Engine/Source/Developer/iOS/MetalShaderFormat/Private/MetalUtils.cpp
+++ b/Engine/Source/Developer/iOS/MetalShaderFormat/Private/MetalUtils.cpp
@@ -867,37 +867,38 @@ void FMetalCodeBackend::PromoteInputsAndOutputsGlobalHalfToFloat(exec_list* Inst
 			{
 			case ir_var_in:
 			{
-							  auto* NewVar = new(State)ir_variable(NewType, Variable->name, ir_var_in);
-							  NewVar->semantic = Variable->semantic;
-							  Variable->insert_before(NewVar);
-							  Variable->name = nullptr;
-							  Variable->semantic = nullptr;
-							  Variable->mode = ir_var_temporary;
-							  Variable->remove();
-							  exec_list Assignments;
-							  Assignments.push_head(Variable);
-							  CreateNewAssignmentsFloat2Half(State, Assignments, Variable, new(State)ir_dereference_variable(NewVar));
-							  EntryPointSig->body.get_head()->insert_before(&Assignments);
-							  break;
+				auto* NewVar = new(State)ir_variable(NewType, Variable->name, ir_var_in);
+				NewVar->semantic = Variable->semantic;
+				Variable->insert_before(NewVar);
+				Variable->name = nullptr;
+				Variable->semantic = nullptr;
+				Variable->mode = ir_var_temporary;
+				Variable->remove();
+				exec_list Assignments;
+				Assignments.push_head(Variable);
+				CreateNewAssignmentsFloat2Half(State, Assignments, Variable, new(State)ir_dereference_variable(NewVar));
+				EntryPointSig->body.get_head()->insert_before(&Assignments);
 			}
+				break;
+
 			case ir_var_out:
 			{
-							   if (Frequency != HSF_PixelShader)
-							   {
-								   auto* NewVar = new(State)ir_variable(NewType, Variable->name, ir_var_out);
-								   NewVar->semantic = Variable->semantic;
-								   Variable->insert_before(NewVar);
-								   Variable->name = nullptr;
-								   Variable->semantic = nullptr;
-								   Variable->mode = ir_var_temporary;
-								   Variable->remove();
-								   exec_list Assignments;
-								   CreateNewAssignmentsHalf2Float(State, Assignments, NewVar, new(State)ir_dereference_variable(Variable));
-								   EntryPointSig->body.push_head(Variable);
-								   EntryPointSig->body.append_list(&Assignments);
-								   break;
-							   }
+				if (Frequency != HSF_PixelShader)
+				{
+				   auto* NewVar = new(State)ir_variable(NewType, Variable->name, ir_var_out);
+				   NewVar->semantic = Variable->semantic;
+				   Variable->insert_before(NewVar);
+				   Variable->name = nullptr;
+				   Variable->semantic = nullptr;
+				   Variable->mode = ir_var_temporary;
+				   Variable->remove();
+				   exec_list Assignments;
+				   CreateNewAssignmentsHalf2Float(State, Assignments, NewVar, new(State)ir_dereference_variable(Variable));
+				   EntryPointSig->body.push_head(Variable);
+				   EntryPointSig->body.append_list(&Assignments);
+				}
 			}
+			   break;
 			}
 		}
 	}
diff --git a/Engine/Source/Programs/CrossCompilerTool/Private/CrossCompilerTool.cpp b/Engine/Source/Programs/CrossCompilerTool/Private/CrossCompilerTool.cpp
--- a/Engine/Source/Programs/CrossCompilerTool/Private/CrossCompilerTool.cpp
+++ b/Engine/Source/Programs/CrossCompilerTool/Private/CrossCompilerTool.cpp
@@ -127,6 +127,14 @@ namespace CCT
 			//Scanner.Dump();
 			return 0;
 		}
+		else
+		{
+			if (!FFileHelper::LoadFileToString(HLSLShaderSource, *RunInfo.InputFile))
+			{
+				UE_LOG(LogCrossCompilerTool, Error, TEXT("Couldn't load Input file '%s'!"), *RunInfo.InputFile);
+				return 1;
+			}
+		}
 
 		ANSICHAR* ShaderSource = 0;
 		ANSICHAR* ErrorLog = 0;
