#UE4 - Metal - missing compute intrinsic

diff --git a/Engine/Source/Developer/iOS/MetalShaderFormat/Private/MetalBackend.cpp b/Engine/Source/Developer/iOS/MetalShaderFormat/Private/MetalBackend.cpp
--- a/Engine/Source/Developer/iOS/MetalShaderFormat/Private/MetalBackend.cpp
+++ b/Engine/Source/Developer/iOS/MetalShaderFormat/Private/MetalBackend.cpp
@@ -20,7 +20,8 @@
 #define _strdup strdup
 #endif
 
-#define GROUP_MEMORY_BARRIER_WITH_GROUP_SYNC "GroupMemoryBarrierWithGroupSync"
+#define GROUP_MEMORY_BARRIER					"GroupMemoryBarrier"
+#define GROUP_MEMORY_BARRIER_WITH_GROUP_SYNC	"GroupMemoryBarrierWithGroupSync"
 
 /**
  * This table must match the ir_expression_operation enum.
@@ -1482,7 +1483,12 @@ protected:
 		else
 		{
 			//@todo-rco: Fix this properly
-			if (!strcmp(call->callee_name(), GROUP_MEMORY_BARRIER_WITH_GROUP_SYNC))
+			if (!strcmp(call->callee_name(), GROUP_MEMORY_BARRIER))
+			{
+				ralloc_asprintf_append(buffer, "threadgroup_barrier(mem_flags::mem_device)");
+				return;
+			}
+			else if (!strcmp(call->callee_name(), GROUP_MEMORY_BARRIER_WITH_GROUP_SYNC))
 			{
 				ralloc_asprintf_append(buffer, "threadgroup_barrier(mem_flags::mem_threadgroup)");
 				return;
@@ -2726,6 +2732,9 @@ void FMetalLanguageSpec::SetupLanguageIntrinsics(_mesa_glsl_parse_state* State,
 
 	// Memory sync/barriers
 	{
+		// GroupMemoryBarrier
+		make_intrinsic_genType(ir, State, GROUP_MEMORY_BARRIER, ir_invalid_opcode, 0, 0, 0, 0);
+
 		// GroupMemoryBarrierWithGroupSync
 		make_intrinsic_genType(ir, State, GROUP_MEMORY_BARRIER_WITH_GROUP_SYNC, ir_invalid_opcode, 0, 0, 0, 0);
 	}
