#UE4 - Metal - Proper sync primitives, add missing #include

diff --git a/Engine/Source/Developer/iOS/MetalShaderFormat/Private/MetalBackend.cpp b/Engine/Source/Developer/iOS/MetalShaderFormat/Private/MetalBackend.cpp
--- a/Engine/Source/Developer/iOS/MetalShaderFormat/Private/MetalBackend.cpp
+++ b/Engine/Source/Developer/iOS/MetalShaderFormat/Private/MetalBackend.cpp
@@ -22,6 +22,10 @@
 
 #define GROUP_MEMORY_BARRIER					"GroupMemoryBarrier"
 #define GROUP_MEMORY_BARRIER_WITH_GROUP_SYNC	"GroupMemoryBarrierWithGroupSync"
+#define DEVICE_MEMORY_BARRIER					"DeviceMemoryBarrier"
+#define DEVICE_MEMORY_BARRIER_WITH_GROUP_SYNC	"DeviceMemoryBarrierWithGroupSync"
+#define ALL_MEMORY_BARRIER						"AllMemoryBarrier"
+#define ALL_MEMORY_BARRIER_WITH_GROUP_SYNC		"AllMemoryBarrierWithGroupSync"
 
 /**
  * This table must match the ir_expression_operation enum.
@@ -405,6 +409,9 @@ protected:
 	// Use packed_ prefix when printing out structs
 	bool bUsePacked;
 
+	// Do we need to add #include <compute_shaders>
+	bool bNeedsComputeInclude;
+
 	/**
 	 * Return true if the type is a multi-dimensional array. Also, track the
 	 * array.
@@ -1483,14 +1490,22 @@ protected:
 		else
 		{
 			//@todo-rco: Fix this properly
-			if (!strcmp(call->callee_name(), GROUP_MEMORY_BARRIER))
+			if (!strcmp(call->callee_name(), GROUP_MEMORY_BARRIER) || !strcmp(call->callee_name(), GROUP_MEMORY_BARRIER_WITH_GROUP_SYNC))
 			{
+				bNeedsComputeInclude = true;
+				ralloc_asprintf_append(buffer, "threadgroup_barrier(mem_flags::mem_threadgroup)");
+				return;
+			}
+			else if (!strcmp(call->callee_name(), DEVICE_MEMORY_BARRIER) || !strcmp(call->callee_name(), DEVICE_MEMORY_BARRIER_WITH_GROUP_SYNC))
+			{
+				bNeedsComputeInclude = true;
 				ralloc_asprintf_append(buffer, "threadgroup_barrier(mem_flags::mem_device)");
 				return;
 			}
-			else if (!strcmp(call->callee_name(), GROUP_MEMORY_BARRIER_WITH_GROUP_SYNC))
+			else if (!strcmp(call->callee_name(), ALL_MEMORY_BARRIER) || !strcmp(call->callee_name(), ALL_MEMORY_BARRIER_WITH_GROUP_SYNC))
 			{
-				ralloc_asprintf_append(buffer, "threadgroup_barrier(mem_flags::mem_threadgroup)");
+				bNeedsComputeInclude = true;
+				ralloc_asprintf_append(buffer, "threadgroup_barrier(mem_flags::mem_device_and_threadgroup)");
 				return;
 			}
 		}
@@ -2438,6 +2453,7 @@ public:
 		, loop_count(0)
 		, bStageInEmitted(false)
 		, bUsePacked(false)
+		, bNeedsComputeInclude(false)
 	{
 		printable_names = hash_table_ctor(32, hash_table_pointer_hash, hash_table_pointer_compare);
 		used_structures = hash_table_ctor(128, hash_table_pointer_hash, hash_table_pointer_compare);
@@ -2482,8 +2498,9 @@ public:
 
 		char* full_buffer = ralloc_asprintf(
 			ParseState,
-			"// Compiled by HLSLCC\n%s\n#include <metal_stdlib>\nusing namespace metal;\n\n%s%s",
+			"// Compiled by HLSLCC\n%s\n#include <metal_stdlib>\n%s\nusing namespace metal;\n\n%s%s",
 			signature,
+			bNeedsComputeInclude ? "#include <metal_compute>" : "",
 			decl_buffer,
 			code_buffer
 			);
@@ -2732,10 +2749,11 @@ void FMetalLanguageSpec::SetupLanguageIntrinsics(_mesa_glsl_parse_state* State,
 
 	// Memory sync/barriers
 	{
-		// GroupMemoryBarrier
 		make_intrinsic_genType(ir, State, GROUP_MEMORY_BARRIER, ir_invalid_opcode, 0, 0, 0, 0);
-
-		// GroupMemoryBarrierWithGroupSync
 		make_intrinsic_genType(ir, State, GROUP_MEMORY_BARRIER_WITH_GROUP_SYNC, ir_invalid_opcode, 0, 0, 0, 0);
+		make_intrinsic_genType(ir, State, DEVICE_MEMORY_BARRIER, ir_invalid_opcode, 0, 0, 0, 0);
+		make_intrinsic_genType(ir, State, DEVICE_MEMORY_BARRIER_WITH_GROUP_SYNC, ir_invalid_opcode, 0, 0, 0, 0);
+		make_intrinsic_genType(ir, State, ALL_MEMORY_BARRIER, ir_invalid_opcode, 0, 0, 0, 0);
+		make_intrinsic_genType(ir, State, ALL_MEMORY_BARRIER_WITH_GROUP_SYNC, ir_invalid_opcode, 0, 0, 0, 0);
 	}
 }
