#UE4 - Metal - Fix for UAVs

diff --git a/Engine/Source/Developer/iOS/MetalShaderFormat/Private/MetalBackend.cpp b/Engine/Source/Developer/iOS/MetalShaderFormat/Private/MetalBackend.cpp
--- a/Engine/Source/Developer/iOS/MetalShaderFormat/Private/MetalBackend.cpp
+++ b/Engine/Source/Developer/iOS/MetalShaderFormat/Private/MetalBackend.cpp
@@ -2365,7 +2365,7 @@ protected:
 				if (Buffers.Buffers[i])
 				{
 					auto* Var = Buffers.Buffers[i]->as_variable();
-					if (!Var->semantic && !Var->type->is_sampler())
+					if (!Var->semantic && !Var->type->is_sampler() && !Var->type->is_image())
 					{
 						ralloc_asprintf_append(buffer, "%s%s(%d)",
 							bFirst ? "// @UniformBlocks: " : ",",
