- Added a Metal shader CC workaround for use of packHalf2x16 in compute shader (likely can be done cleaner - rco!)

diff --git a/Engine/Source/Developer/iOS/MetalShaderFormat/Private/MetalBackend.cpp b/Engine/Source/Developer/iOS/MetalShaderFormat/Private/MetalBackend.cpp
--- a/Engine/Source/Developer/iOS/MetalShaderFormat/Private/MetalBackend.cpp
+++ b/Engine/Source/Developer/iOS/MetalShaderFormat/Private/MetalBackend.cpp
@@ -1510,7 +1510,14 @@ protected:
 			}
 		}
 
-		ralloc_asprintf_append(buffer, "%s(", call->callee_name());
+		if (!strcmp(call->callee_name(), "packHalf2x16"))
+		{
+			ralloc_asprintf_append(buffer, "pack_float_to_snorm2x16(");
+		}
+		else
+		{
+			ralloc_asprintf_append(buffer, "%s(", call->callee_name());
+		}
 		bool bPrintComma = false;
 		foreach_iter(exec_list_iterator, iter, *call)
 		{
