SSH and IOS remote builds
- Converted a few properties to be available in the editor.
- Part of UEPLAT-25

diff --git a/Engine/Source/Developer/iOS/IOSPlatformEditor/Private/IOSTargetSettingsCustomization.cpp b/Engine/Source/Developer/iOS/IOSPlatformEditor/Private/IOSTargetSettingsCustomization.cpp
--- a/Engine/Source/Developer/iOS/IOSPlatformEditor/Private/IOSTargetSettingsCustomization.cpp
+++ b/Engine/Source/Developer/iOS/IOSPlatformEditor/Private/IOSTargetSettingsCustomization.cpp
@@ -7,6 +7,7 @@
 #include "PropertyEditing.h"
 #include "DesktopPlatformModule.h"
 #include "MainFrame.h"
+#include "IDetailPropertyRow.h"
 
 #include "ScopedTransaction.h"
 #include "SExternalImageReference.h"
@@ -24,6 +25,11 @@
 
 //////////////////////////////////////////////////////////////////////////
 // FIOSTargetSettingsCustomization
+namespace FIOSTargetSettingsCustomizationConstants
+{
+	const FText DisabledTip = LOCTEXT("GitHubSourceRequiredToolTip", "This requires GitHub source.");
+}
+
 
 TSharedRef<IDetailCustomization> FIOSTargetSettingsCustomization::MakeInstance()
 {
@@ -71,6 +77,8 @@ void FIOSTargetSettingsCustomization::CustomizeDetails(IDetailLayoutBuilder& Det
 	BuildPListSection(DetailLayout);
 
 	BuildIconSection(DetailLayout);
+
+	BuildRemoteBuildingSection(DetailLayout);
 }
 
 void FIOSTargetSettingsCustomization::UpdateStatus()
@@ -299,18 +307,72 @@ void FIOSTargetSettingsCustomization::BuildPListSection(IDetailLayoutBuilder& De
 
 	SETUP_PLIST_PROP(MinimumiOSVersion, OSInfoCategory, LOCTEXT("MinimumiOSVersionToolTip", "Minimum iOS version this game supports"));
 
-	const FText DisabledTip = LOCTEXT("GitHubSourceRequiredToolTip", "This requires GitHub source.");
-	SETUP_NONROCKET_PROP(bDevForArmV7, BuildCategory, LOCTEXT("DevForArmV7ToolTip", "Enable ArmV7 support? (this will be used if all type are unchecked)"), DisabledTip);
-	SETUP_NONROCKET_PROP(bDevForArm64, BuildCategory, LOCTEXT("DevForArm64ToolTip", "Enable Arm64 support?"), DisabledTip);
-	SETUP_NONROCKET_PROP(bDevForArmV7S, BuildCategory, LOCTEXT("DevForArmV7SToolTip", "Enable ArmV7s support?"), DisabledTip);
-	SETUP_NONROCKET_PROP(bShipForArmV7, BuildCategory, LOCTEXT("ShipForArmV7ToolTip", "Enable ArmV7 support? (this will be used if all type are unchecked)"), DisabledTip);
-	SETUP_NONROCKET_PROP(bShipForArm64, BuildCategory, LOCTEXT("ShipForArm64ToolTip", "Enable Arm64 support?"), DisabledTip);
-	SETUP_NONROCKET_PROP(bShipForArmV7S, BuildCategory, LOCTEXT("ShipForArmV7SToolTip", "Enable ArmV7s support?"), DisabledTip);
+	SETUP_NONROCKET_PROP(bDevForArmV7, BuildCategory, LOCTEXT("DevForArmV7ToolTip", "Enable ArmV7 support? (this will be used if all type are unchecked)"), FIOSTargetSettingsCustomizationConstants::DisabledTip);
+	SETUP_NONROCKET_PROP(bDevForArm64, BuildCategory, LOCTEXT("DevForArm64ToolTip", "Enable Arm64 support?"), FIOSTargetSettingsCustomizationConstants::DisabledTip);
+	SETUP_NONROCKET_PROP(bDevForArmV7S, BuildCategory, LOCTEXT("DevForArmV7SToolTip", "Enable ArmV7s support?"), FIOSTargetSettingsCustomizationConstants::DisabledTip);
+	SETUP_NONROCKET_PROP(bShipForArmV7, BuildCategory, LOCTEXT("ShipForArmV7ToolTip", "Enable ArmV7 support? (this will be used if all type are unchecked)"), FIOSTargetSettingsCustomizationConstants::DisabledTip);
+	SETUP_NONROCKET_PROP(bShipForArm64, BuildCategory, LOCTEXT("ShipForArm64ToolTip", "Enable Arm64 support?"), FIOSTargetSettingsCustomizationConstants::DisabledTip);
+	SETUP_NONROCKET_PROP(bShipForArmV7S, BuildCategory, LOCTEXT("ShipForArmV7SToolTip", "Enable ArmV7s support?"), FIOSTargetSettingsCustomizationConstants::DisabledTip);
 
 #undef SETUP_PLIST_PROP
 #undef SETUP_NONROCKET_PROP
 }
 
+
+void FIOSTargetSettingsCustomization::BuildRemoteBuildingSection(IDetailLayoutBuilder& DetailLayout)
+{
+	IDetailCategoryBuilder& BuildCategory = DetailLayout.EditCategory(TEXT("Build"));
+
+	// Sub group we wish to add remote building options to.
+	FText RemoteBuildingGroupName = LOCTEXT("RemoteBuildingGroupName", "Remote Build Options");
+	IDetailGroup& RemoteBuildingGroup = BuildCategory.AddGroup(*RemoteBuildingGroupName.ToString(), RemoteBuildingGroupName, false);
+
+	// Remote Server Name Property
+	TSharedRef<IPropertyHandle> RemoteServerNamePropertyHandle = DetailLayout.GetProperty(GET_MEMBER_NAME_CHECKED(UIOSRuntimeSettings, RemoteServerName));
+	IDetailPropertyRow& RemoteServerNamePropertyRow = RemoteBuildingGroup.AddPropertyRow(RemoteServerNamePropertyHandle);
+	RemoteServerNamePropertyRow
+		.EditCondition(SetupForPlatformAttribute, NULL)
+		.IsEnabled(!FRocketSupport::IsRocket())
+		.ToolTip(!FRocketSupport::IsRocket() ? LOCTEXT("RemoteServerNameToolTip", "The name or ip address of the remote mac which will be used to build IOS") : FIOSTargetSettingsCustomizationConstants::DisabledTip);
+
+	
+	// Add Use RSync Property
+	TSharedRef<IPropertyHandle> UseRSyncPropertyHandle = DetailLayout.GetProperty(GET_MEMBER_NAME_CHECKED(UIOSRuntimeSettings, bUseRSync));
+	IDetailPropertyRow& UseRSyncPropertRow = RemoteBuildingGroup.AddPropertyRow(UseRSyncPropertyHandle);
+	UseRSyncPropertRow
+		.EditCondition(SetupForPlatformAttribute, NULL)
+		.IsEnabled(!FRocketSupport::IsRocket())
+		.ToolTip(!FRocketSupport::IsRocket() ? LOCTEXT("UseRSyncToolTip", "Use RSync instead of RPCUtility") : FIOSTargetSettingsCustomizationConstants::DisabledTip);
+
+	
+	// Add RSync Username Property
+	TSharedRef<IPropertyHandle> RSyncUsernamePropertyHandle = DetailLayout.GetProperty(GET_MEMBER_NAME_CHECKED(UIOSRuntimeSettings, RSyncUsername));
+	IDetailPropertyRow& RSyncUsernamePropertyRow = RemoteBuildingGroup.AddPropertyRow(RSyncUsernamePropertyHandle);
+	RSyncUsernamePropertyRow
+		.EditCondition(SetupForPlatformAttribute, NULL)
+		.IsEnabled(!FRocketSupport::IsRocket())
+		.ToolTip(!FRocketSupport::IsRocket() ? LOCTEXT("RSyncUsernameToolTip", "The username of the mac user that matches the specified SSH Key.") : FIOSTargetSettingsCustomizationConstants::DisabledTip);
+
+
+	// Add existing SSH path label.
+	TSharedRef<IPropertyHandle> SSHPrivateKeyLocationPropertyHandle = DetailLayout.GetProperty(GET_MEMBER_NAME_CHECKED(UIOSRuntimeSettings, SSHPrivateKeyLocation));
+	IDetailPropertyRow& SSHPrivateKeyLocationPropertyRow = RemoteBuildingGroup.AddPropertyRow(SSHPrivateKeyLocationPropertyHandle);
+	SSHPrivateKeyLocationPropertyRow
+		.EditCondition(SetupForPlatformAttribute, NULL)
+		.IsEnabled(!FRocketSupport::IsRocket())
+		.ToolTip(!FRocketSupport::IsRocket() ? LOCTEXT("SSHPrivateKeyLocationToolTip", "The existing location of an SSH Key found by UE4.") : FIOSTargetSettingsCustomizationConstants::DisabledTip);
+
+
+	// Add SSH override path
+	TSharedRef<IPropertyHandle> SSHPrivateKeyOverridePathPropertyHandle = DetailLayout.GetProperty(GET_MEMBER_NAME_CHECKED(UIOSRuntimeSettings, SSHPrivateKeyOverridePath));
+	IDetailPropertyRow& SSHPrivateKeyOverridePathPropertyRow = RemoteBuildingGroup.AddPropertyRow(SSHPrivateKeyOverridePathPropertyHandle);
+	SSHPrivateKeyOverridePathPropertyRow
+		.EditCondition(SetupForPlatformAttribute, NULL)
+		.IsEnabled(!FRocketSupport::IsRocket())
+		.ToolTip(!FRocketSupport::IsRocket() ? LOCTEXT("SSHPrivateKeyOverridePathToolTip", "Override the existing SSH Private Key with one from a specified location.") : FIOSTargetSettingsCustomizationConstants::DisabledTip);
+}
+
+
 void FIOSTargetSettingsCustomization::BuildIconSection(IDetailLayoutBuilder& DetailLayout)
 {
 	IDetailCategoryBuilder& RequiredIconCategory = DetailLayout.EditCategory(TEXT("Required Icons"));
@@ -565,7 +627,7 @@ FReply FIOSTargetSettingsCustomization::OnInstallProvisionClicked()
 
 	if ( bOpened )
 	{
-        ProvisionPath = FPaths::ConvertRelativePathToFull(OpenFilenames[0]);
+		ProvisionPath = FPaths::ConvertRelativePathToFull(OpenFilenames[0]);
 
 		// see if the provision is already installed
 		FString DestName = FPaths::GetBaseFilename(ProvisionPath);
@@ -644,7 +706,7 @@ FReply FIOSTargetSettingsCustomization::OnInstallCertificateClicked()
 
 	if ( bOpened )
 	{
-        CertPath = FPaths::ConvertRelativePathToFull(OpenFilenames[0]);
+		CertPath = FPaths::ConvertRelativePathToFull(OpenFilenames[0]);
 #if PLATFORM_MAC
 		FString CmdExe = TEXT("/bin/sh");
 		FString ScriptPath = FPaths::ConvertRelativePathToFull(FPaths::EngineDir() / TEXT("Build/BatchFiles/Mac/RunMono.sh"));
diff --git a/Engine/Source/Developer/iOS/IOSPlatformEditor/Private/IOSTargetSettingsCustomization.h b/Engine/Source/Developer/iOS/IOSPlatformEditor/Private/IOSTargetSettingsCustomization.h
--- a/Engine/Source/Developer/iOS/IOSPlatformEditor/Private/IOSTargetSettingsCustomization.h
+++ b/Engine/Source/Developer/iOS/IOSPlatformEditor/Private/IOSTargetSettingsCustomization.h
@@ -44,6 +44,7 @@ private:
 
 	void BuildPListSection(IDetailLayoutBuilder& DetailLayout);
 	void BuildIconSection(IDetailLayoutBuilder& DetailLayout);
+	void BuildRemoteBuildingSection(IDetailLayoutBuilder& DetailLayout);
 
 	// Navigates to the plist in explorer or finder
 	FReply OpenPlistFolder();
diff --git a/Engine/Source/Editor/DetailCustomizations/Private/DetailCustomizations.cpp b/Engine/Source/Editor/DetailCustomizations/Private/DetailCustomizations.cpp
--- a/Engine/Source/Editor/DetailCustomizations/Private/DetailCustomizations.cpp
+++ b/Engine/Source/Editor/DetailCustomizations/Private/DetailCustomizations.cpp
@@ -178,6 +178,7 @@ void FDetailCustomizationsModule::RegisterPropertyTypeCustomizations()
 	RegisterCustomPropertyTypeLayout("Margin", FOnGetPropertyTypeCustomizationInstance::CreateStatic(&FMarginStructCustomization::MakeInstance));
 	RegisterCustomPropertyTypeLayout("DirectoryPath", FOnGetPropertyTypeCustomizationInstance::CreateStatic(&FDirectoryPathStructCustomization::MakeInstance));
 	RegisterCustomPropertyTypeLayout("FilePath", FOnGetPropertyTypeCustomizationInstance::CreateStatic(&FFilePathStructCustomization::MakeInstance));
+	RegisterCustomPropertyTypeLayout("IOSBuildResourceFilePath", FOnGetPropertyTypeCustomizationInstance::CreateStatic(&FFilePathStructCustomization::MakeInstance));
 	RegisterCustomPropertyTypeLayout("InputAxisConfigEntry", FOnGetPropertyTypeCustomizationInstance::CreateStatic(&FInputAxisConfigCustomization::MakeInstance));
 	RegisterCustomPropertyTypeLayout("InputActionKeyMapping", FOnGetPropertyTypeCustomizationInstance::CreateStatic(&FInputActionMappingCustomization::MakeInstance));
 	RegisterCustomPropertyTypeLayout("InputAxisKeyMapping", FOnGetPropertyTypeCustomizationInstance::CreateStatic(&FInputAxisMappingCustomization::MakeInstance));
diff --git a/Engine/Source/Programs/UnrealBuildTool/ToolChain/RemoteToolChain.cs b/Engine/Source/Programs/UnrealBuildTool/ToolChain/RemoteToolChain.cs
--- a/Engine/Source/Programs/UnrealBuildTool/ToolChain/RemoteToolChain.cs
+++ b/Engine/Source/Programs/UnrealBuildTool/ToolChain/RemoteToolChain.cs
@@ -186,6 +186,48 @@ namespace UnrealBuildTool
 			return InPath;
 		}
 
+		// Look for any build options in the engine config file.
+		public override void ParseProjectSettings()
+		{
+			base.ParseProjectSettings();
+
+			ConfigCacheIni Ini = new ConfigCacheIni(UnrealTargetPlatform.IOS, "Engine", UnrealBuildTool.GetUProjectPath());
+			Ini.GetString("/Script/IOSRuntimeSettings.IOSRuntimeSettings", "RemoteServerName", out RemoteServerName);
+
+			bool bUseRSync = false;
+			if (Ini.GetBool("/Script/IOSRuntimeSettings.IOSRuntimeSettings", "bUseRSync", out bUseRSync))
+			{
+				bUseRPCUtil = !bUseRSync;
+				Ini.GetString("/Script/IOSRuntimeSettings.IOSRuntimeSettings", "RSyncUsername", out RSyncUsername);
+
+				string ConfigKeyPath;
+				if (Ini.GetString("/Script/IOSRuntimeSettings.IOSRuntimeSettings", "SSHPrivateKeyOverridePath", out ConfigKeyPath))
+				{
+					if (File.Exists(ConfigKeyPath))
+					{
+						SSHPrivateKeyOverridePath = ConfigKeyPath;
+					}
+				}
+			}
+		}
+
+		// Gather a users root path from the remote server. Should only be called once.
+		public static void SetUserDevRootFromServer()
+		{
+			Hashtable Results = SSHCommand("/", "echo $HOME", null);
+
+			if (Results == null)
+			{
+				Log.TraceInformation("UserDevRoot Command failed to execute!");
+			}
+			else if (Results["CommandOutput"] != null)
+			{
+				// pass back the string
+				string HomeLocation = Results["CommandOutput"] as string;
+				UserDevRootMac = HomeLocation + UserDevRootMac;
+			}
+		}
+
         // Do any one-time, global initialization for the tool chain
 		static int InitializationErrorCode = 0;
 		private static int InitializeRemoteExecution()
@@ -264,22 +306,31 @@ namespace UnrealBuildTool
 					// we need the RemoteServerName and the Username to find the private key
 					ResolvedRSyncUsername = ResolveString(RSyncUsername, false);
 
+					bool bFoundOverrideSSHPrivateKey = false;
+
 					// if the override path is set, just use it directly
 					if (!string.IsNullOrEmpty(SSHPrivateKeyOverridePath))
 					{
 						ResolvedSSHPrivateKey = ResolveString(SSHPrivateKeyOverridePath, true);
+						
+						bFoundOverrideSSHPrivateKey = File.Exists(ResolvedSSHPrivateKey);
+						
 						// make sure it exists
-						if (!File.Exists(ResolvedSSHPrivateKey))
+						if (!bFoundOverrideSSHPrivateKey)
 						{
-							throw new BuildException("An SSHKey override was specified [" + SSHPrivateKeyOverridePath + "] but it doesn't exist. Can't continue...");
+							Log.TraceWarning("An SSHKey override was specified [" + SSHPrivateKeyOverridePath + "] but it doesn't exist. Looking elsewhere...");
 						}
 					}
-					else
+
+					if (!bFoundOverrideSSHPrivateKey)
 					{
 						// all the places to look for a key
 						string[] Locations = new string[] {
 							Path.Combine(Environment.GetFolderPath(Environment.SpecialFolder.ApplicationData), "Unreal Engine", "UnrealBuildTool"),
 							Path.Combine(Environment.GetFolderPath(Environment.SpecialFolder.Personal), "Unreal Engine", "UnrealBuildTool"),
+							Path.Combine(UnrealBuildTool.GetUProjectPath(), "Build", "NotForLicensees"),
+							Path.Combine(UnrealBuildTool.GetUProjectPath(), "Build", "NoRedist"),
+							Path.Combine(UnrealBuildTool.GetUProjectPath(), "Build"),
 							Path.Combine(BuildConfiguration.RelativeEnginePath, "Build", "NotForLicensees"),
 							Path.Combine(BuildConfiguration.RelativeEnginePath, "Build", "NoRedist"),
 							Path.Combine(BuildConfiguration.RelativeEnginePath, "Build"),
@@ -292,11 +343,17 @@ namespace UnrealBuildTool
 							if (File.Exists(KeyPath))
 							{
 								ResolvedSSHPrivateKey = KeyPath;
+								bFoundOverrideSSHPrivateKey = true;
 								break;
 							}
 						}
 					}
 
+					if (!bFoundOverrideSSHPrivateKey)
+					{
+						throw new BuildException("An SSHKey was required, but one cannot be found. Can't continue...");
+					}
+
 					// resolve the rest of the strings
 					ResolvedRSyncExe = ResolveString(RSyncExe, true);
 					ResolvedSSHExe = ResolveString(SSHExe, true);
@@ -349,6 +406,9 @@ namespace UnrealBuildTool
 
 			// connect to server
 			InitializeRemoteExecution();
+
+			// Setup root directory to use.
+			SetUserDevRootFromServer();
 		}
 
         /** Converts the passed in path from UBT host to compiler native format. */
@@ -621,7 +681,7 @@ namespace UnrealBuildTool
 				// --exclude='*'  ??? why???
 				RsyncProcess.StartInfo.FileName = ResolvedRSyncExe;
 				RsyncProcess.StartInfo.Arguments = string.Format(
-					"-vzae \"{0}\" --rsync-path=\"mkdir -p {2} && rsync\" --delete --files-from=\"{4}\" --include-from=\"{5}\" --include='*/' --exclude='*.o' --exclude='Timestamp' {1} {6}@{3}:'{2}'",
+					"-vzae \"{0}\" --rsync-path=\"mkdir -p {2} && rsync\" --chmod=ug=rwX,o=rxX --delete --files-from=\"{4}\" --include-from=\"{5}\" --include='*/' --exclude='*.o' --exclude='Timestamp' {1} {6}@{3}:'{2}'",
 					ResolvedRsyncAuthentication,
 					CygRootPath,
 					RemotePath,
@@ -662,7 +722,7 @@ namespace UnrealBuildTool
 			// make simple rsync commandline to send a file
 			RsyncProcess.StartInfo.FileName = ResolvedRSyncExe;
 			RsyncProcess.StartInfo.Arguments = string.Format(
-				"-zae \"{0}\" --rsync-path=\"mkdir -p '{1}' && rsync\" \"{2}\" {3}@{4}:\"{1}/{5}\"",
+				"-zae \"{0}\" --rsync-path=\"mkdir -p {1} && rsync\" \"{2}\" {3}@{4}:\"{1}/{5}\"",
 				ResolvedRsyncAuthentication,
 				RemoteDir,
 				ConvertPathToCygwin(LocalPath),
diff --git a/Engine/Source/Runtime/IOS/IOSRuntimeSettings/Classes/IOSRuntimeSettings.h b/Engine/Source/Runtime/IOS/IOSRuntimeSettings/Classes/IOSRuntimeSettings.h
--- a/Engine/Source/Runtime/IOS/IOSRuntimeSettings/Classes/IOSRuntimeSettings.h
+++ b/Engine/Source/Runtime/IOS/IOSRuntimeSettings/Classes/IOSRuntimeSettings.h
@@ -34,6 +34,55 @@ UENUM()
 	IOS_8 = 8 UMETA(DisplayName="8.0"),
 };
 
+
+/**
+ *	IOS Build resource file struct, used to serialize filepaths to the configs for use in the build system,
+ */
+USTRUCT()
+struct FIOSBuildResourceFilePath
+{
+	GENERATED_USTRUCT_BODY()
+	
+	/**
+	 * Custom export item used to serialize FIOSBuildResourceFilePath types as only a filename, no garland.
+	 */
+	bool ExportTextItem(FString& ValueStr, FIOSBuildResourceFilePath const& DefaultValue, UObject* Parent, int32 PortFlags, UObject* ExportRootScope) const
+	{
+		ValueStr += FilePath;
+		return true;
+	}
+
+	/**
+	 * Custom import item used to parse ini entries straight into the filename.
+	 */
+	bool ImportTextItem(const TCHAR*& Buffer, int32 PortFlags, UObject* Parent, FOutputDevice* ErrorText)
+	{
+		FilePath = Buffer;
+		return true;
+	}
+
+	/**
+	 * The path to the file.
+	 */
+	UPROPERTY(EditAnywhere, Category = FilePath)
+	FString FilePath;
+};
+
+/**
+ *	Setup our resource filepath to make it easier to parse in UBT
+ */
+template<>
+struct TStructOpsTypeTraits<FIOSBuildResourceFilePath> : public TStructOpsTypeTraitsBase
+{
+	enum
+	{
+		WithExportTextItem = true,
+		WithImportTextItem = true,
+	};
+};
+
+
+
 /**
  * Implements the settings for the iOS target platform.
  */
@@ -82,6 +131,26 @@ public:
 	// Enable ArmV7s support for shipping build? [CURRENTLY AVAILABLE ON MAC ONLY OR FULL SOURCE BUILD]
 	UPROPERTY(GlobalConfig, EditAnywhere, Category = Build, meta = (DisplayName = "Support armv7s in Shipping"))
 	bool bShipForArmV7S;
+	
+	// The name or ip address of the remote mac which will be used to build IOS
+	UPROPERTY(GlobalConfig, EditAnywhere, Category = "Build")
+	FString RemoteServerName;
+
+	// Enable the use of RSync for remote builds on a mac
+	UPROPERTY(GlobalConfig, EditAnywhere, Category = "Build", meta = (DisplayName = "Use RSync for IOS build"))
+	bool bUseRSync;
+
+	// The mac users name which matches the SSH Private Key, for remote builds using RSync.
+	UPROPERTY(GlobalConfig, EditAnywhere, Category = "Build", meta = (EditCondition = "bUseRSync", DisplayName = "User for RSync enabled builds."))
+	FString RSyncUsername;
+
+	// The path of the ssh permissions key to be used when connecting to the remote server.
+	UPROPERTY(VisibleAnywhere, Category = "Build", meta = (DisplayName = "Existing SSH permissions file"))
+	FString SSHPrivateKeyLocation;
+
+	// The path of the ssh permissions key to be used when connecting to the remote server.
+	UPROPERTY(GlobalConfig, EditAnywhere, Category = "Build", meta = (EditCondition = "bUseRSync", DisplayName = "Override existing SSH permissions file"))
+	FIOSBuildResourceFilePath SSHPrivateKeyOverridePath;
 
 	// Does the application support portrait orientation?
 	UPROPERTY(GlobalConfig, EditAnywhere, Category = DeviceOrientations)
@@ -134,6 +203,7 @@ public:
 #if WITH_EDITOR
 	// UObject interface
 	virtual void PostEditChangeProperty(struct FPropertyChangedEvent& PropertyChangedEvent) override;
+	virtual void PostInitProperties() override;
 	// End of UObject interface
 #endif
 };
diff --git a/Engine/Source/Runtime/IOS/IOSRuntimeSettings/Private/IOSRuntimeSettings.cpp b/Engine/Source/Runtime/IOS/IOSRuntimeSettings/Private/IOSRuntimeSettings.cpp
--- a/Engine/Source/Runtime/IOS/IOSRuntimeSettings/Private/IOSRuntimeSettings.cpp
+++ b/Engine/Source/Runtime/IOS/IOSRuntimeSettings/Private/IOSRuntimeSettings.cpp
@@ -23,6 +23,7 @@ UIOSRuntimeSettings::UIOSRuntimeSettings(const FObjectInitializer& ObjectInitial
 	bShipForArmV7 = true;
 	bShipForArm64 = true;
 	bShipForArmV7S = false;
+	bUseRSync = false;
 }
 
 #if WITH_EDITOR
@@ -52,4 +53,36 @@ void UIOSRuntimeSettings::PostEditChangeProperty(struct FPropertyChangedEvent& P
 		bShipForArmV7 = true;
 	}
 }
+
+
+void UIOSRuntimeSettings::PostInitProperties()
+{
+	Super::PostInitProperties();
+
+	// We can have a look for potential keys
+	if (!RemoteServerName.IsEmpty() && !RSyncUsername.IsEmpty())
+	{
+		const FString DefaultKeyFilename = TEXT("RemoteToolChainPrivate.key");
+		const FString RelativeFilePathLocation = FPaths::Combine(TEXT("SSHKeys"), *RemoteServerName, *RSyncUsername, *DefaultKeyFilename);
+
+		TArray<FString> PossibleKeyLocations;
+		PossibleKeyLocations.Add(FPaths::Combine(*FPaths::GameDir(), TEXT("Build"), TEXT("NotForLicensees"), *RelativeFilePathLocation));
+		PossibleKeyLocations.Add(FPaths::Combine(*FPaths::GameDir(), TEXT("Build"), TEXT("NoRedist"), *RelativeFilePathLocation));
+		PossibleKeyLocations.Add(FPaths::Combine(*FPaths::GameDir(), TEXT("Build"), *RelativeFilePathLocation));
+		PossibleKeyLocations.Add(FPaths::Combine(*FPaths::EngineDir(), TEXT("Build"), TEXT("NotForLicensees"), *RelativeFilePathLocation));
+		PossibleKeyLocations.Add(FPaths::Combine(*FPaths::EngineDir(), TEXT("Build"), TEXT("NoRedist"), *RelativeFilePathLocation));
+		PossibleKeyLocations.Add(FPaths::Combine(*FPaths::EngineDir(), TEXT("Build"), *RelativeFilePathLocation));
+
+		// Find a potential path that we will use if the user hasn't overridden.
+		// For information purposes only
+		for (const FString& NextLocation : PossibleKeyLocations)
+		{
+			if (IFileManager::Get().FileSize(*NextLocation) > 0)
+			{
+				SSHPrivateKeyLocation = NextLocation;
+				break;
+			}
+		}
+	}
+}
 #endif
