UE4 - Metal - Fix UAV indices

diff --git a/Engine/Source/Developer/iOS/MetalShaderFormat/Private/MetalBackend.cpp b/Engine/Source/Developer/iOS/MetalShaderFormat/Private/MetalBackend.cpp
--- a/Engine/Source/Developer/iOS/MetalShaderFormat/Private/MetalBackend.cpp
+++ b/Engine/Source/Developer/iOS/MetalShaderFormat/Private/MetalBackend.cpp
@@ -2077,41 +2077,28 @@ protected:
 		}
 	}
 
-	bool PrintPackedUniforms(bool bPrintArrayType, char ArrayType, _mesa_glsl_parse_state::TUniformList& Uniforms, bool bFlattenUniformBuffers, bool NeedsComma)
-		{
+	void PrintImages(_mesa_glsl_parse_state::TUniformList& Uniforms)
+	{
 		bool bPrintHeader = true;
+		bool bNeedsComma = false;
 		for (_mesa_glsl_parse_state::TUniformList::iterator Iter = Uniforms.begin(); Iter != Uniforms.end(); ++Iter)
 		{
 			glsl_packed_uniform& Uniform = *Iter;
-			if (!bFlattenUniformBuffers || Uniform.CB_PackedSampler.empty())
-			{
-				if (bPrintArrayType && bPrintHeader)
-				{
-					ralloc_asprintf_append(buffer, "%s%c[",
-						NeedsComma ? "," : "",
-						ArrayType);
-					bPrintHeader = false;
-					NeedsComma = false;
-				}
-				ralloc_asprintf_append(
-					buffer,
-					"%s%s(%u:%u)",
-					NeedsComma ? "," : "",
-					Uniform.Name.c_str(),
-					Uniform.offset,
-					Uniform.num_components
-					);
-				NeedsComma = true;
-				}
-			}
-
-		if (bPrintArrayType && !bPrintHeader)
-			{
-			ralloc_asprintf_append(buffer, "]");
-			}
-
-		return NeedsComma;
+			ANSICHAR Name[32];
+			FCStringAnsi::Sprintf(Name, "%si%d", glsl_variable_tag_from_parser_target(Frequency), Uniform.offset);
+			int32 Offset = Buffers.GetIndex(Name);
+			check(Offset != -1);
+			ralloc_asprintf_append(
+				buffer,
+				"%s%s(%u:%u)",
+				bNeedsComma ? "," : "",
+				Uniform.Name.c_str(),
+				Offset,
+				Uniform.num_components
+				);
+			bNeedsComma = true;
 		}
+	}
 
 	void PrintPackedGlobals(_mesa_glsl_parse_state* State)
 	{
@@ -2404,13 +2391,7 @@ protected:
 			if (!state->GlobalPackedArraysMap[EArrayType_Image].empty())
 			{
 				ralloc_asprintf_append(buffer, "// @UAVs: ");
-				PrintPackedUniforms(
-					false,
-					EArrayType_Image,
-					state->GlobalPackedArraysMap[EArrayType_Image],
-					false,
-					false
-					);
+				PrintImages(state->GlobalPackedArraysMap[EArrayType_Image]);
 				ralloc_asprintf_append(buffer, "\n");
 			}
 		}
diff --git a/Engine/Source/Developer/iOS/MetalShaderFormat/Private/MetalUtils.h b/Engine/Source/Developer/iOS/MetalShaderFormat/Private/MetalUtils.h
--- a/Engine/Source/Developer/iOS/MetalShaderFormat/Private/MetalUtils.h
+++ b/Engine/Source/Developer/iOS/MetalShaderFormat/Private/MetalUtils.h
@@ -53,6 +53,20 @@ struct FBuffers
 		return -1;
 	}
 
+	int GetIndex(const std::string& Name)
+	{
+		for (int i = 0, n = Buffers.Num(); i < n; ++i)
+		{
+			auto* Var = Buffers[i] ? Buffers[i]->as_variable() : nullptr;
+			if (Var && Var->name && Var->name == Name)
+			{
+				return i;
+			}
+		}
+
+		return -1;
+	}
+
 	void SortBuffers()
 	{
 		TArray<class ir_instruction*> AllBuffers;
