Quick fix for packing 2 halves into a uint (better fix will be handling all packing modes)

diff --git a/Engine/Source/Developer/iOS/MetalShaderFormat/Private/MetalBackend.cpp b/Engine/Source/Developer/iOS/MetalShaderFormat/Private/MetalBackend.cpp
--- a/Engine/Source/Developer/iOS/MetalShaderFormat/Private/MetalBackend.cpp
+++ b/Engine/Source/Developer/iOS/MetalShaderFormat/Private/MetalBackend.cpp
@@ -1500,7 +1500,7 @@ protected:
 
 		if (!strcmp(call->callee_name(), "packHalf2x16"))
 		{
-			ralloc_asprintf_append(buffer, "pack_float_to_snorm2x16(");
+			ralloc_asprintf_append(buffer, "as_type<uint>(half2(");
 		}
 		else
 		{
@@ -1518,6 +1518,11 @@ protected:
 			bPrintComma = true;
 		}
 		ralloc_asprintf_append(buffer, ")");
+
+		if (!strcmp(call->callee_name(), "packHalf2x16"))
+		{
+			ralloc_asprintf_append(buffer, ")");			
+		}
 	}
 
 	virtual void visit(ir_return *ret) override
