UE4 - Metal - Fix for SVPosition.w

diff --git a/Engine/Shaders/Common.usf b/Engine/Shaders/Common.usf
--- a/Engine/Shaders/Common.usf
+++ b/Engine/Shaders/Common.usf
@@ -9,7 +9,12 @@
 #if PS4_PROFILE  
 	// always #include PS4Common.usf so it can #define override anything in any other included file.
 	#include "PS4/PS4Common.usf"
-#endif 
+#endif
+
+#if METAL_PROFILE || METAL_MRT_PROFILE
+	// Helps with iteration when changing Metal shader code generation backend.
+	#include "MetalCommon.usf"
+#endif
 
 #ifndef FORCE_FLOATS
 #define FORCE_FLOATS 0
diff --git a/Engine/Shaders/MetalCommon.usf b/Engine/Shaders/MetalCommon.usf
new file mode 100644
--- /dev/null
+++ b/Engine/Shaders/MetalCommon.usf
@@ -0,0 +1,10 @@
+// Copyright 1998-2015 Epic Games, Inc. All Rights Reserved.
+
+/*=============================================================================
+	MetalCommon.usf: Common Metal shader code
+=============================================================================*/
+
+// Update this GUID to improve shader recompilation for Metal only shaders
+// GUID = {06C79EC9-644E-4D22-ADA9-47DD4BB343D9}
+  
+#pragma once
diff --git a/Engine/Source/Developer/iOS/MetalShaderFormat/Private/MetalUtils.cpp b/Engine/Source/Developer/iOS/MetalShaderFormat/Private/MetalUtils.cpp
--- a/Engine/Source/Developer/iOS/MetalShaderFormat/Private/MetalUtils.cpp
+++ b/Engine/Source/Developer/iOS/MetalShaderFormat/Private/MetalUtils.cpp
@@ -282,15 +282,39 @@ namespace MetalUtils
 			{
 				if (SystemValues[i].Mode == ir_var_in && FCStringAnsi::Stricmp(SystemValues[i].HlslSemantic, Semantic) == 0)
 				{
-					ir_variable* Variable = new(ParseState)ir_variable(
-						SystemValues[i].Type, SystemValues[i].MetalName, ir_var_in);
+					ir_variable* Variable = new(ParseState) ir_variable(SystemValues[i].Type, SystemValues[i].MetalName, ir_var_in);
 					Variable->semantic = SystemValues[i].MetalSemantic;
 					Variable->read_only = true;
 					Variable->origin_upper_left = false;
 					DeclInstructions->push_tail(Variable);
 					ParseState->symbols->add_variable(Variable);
-					ir_dereference_variable* VariableDeref = new(ParseState)ir_dereference_variable(Variable);
+					ir_dereference_variable* VariableDeref = new(ParseState) ir_dereference_variable(Variable);
+					if (!FCStringAnsi::Stricmp(Semantic, "SV_Position") && Frequency == HSF_PixelShader)
+					{
+						// UE4 requires w instead of 1/w in SVPosition
+						auto* TempVariable = new(ParseState) ir_variable(Variable->type, nullptr, ir_var_temporary);
+						DeclInstructions->push_tail(TempVariable);
+
+						// Assign input to this variable
+						auto* TempVariableDeref = new(ParseState) ir_dereference_variable(TempVariable);
+						DeclInstructions->push_tail(
+							new(ParseState) ir_assignment(TempVariableDeref, VariableDeref)
+							);
+
+						// TempVariable.w = ( 1.0f / TempVariable.w );
+						DeclInstructions->push_tail(
+							new(ParseState) ir_assignment(
+								new(ParseState) ir_swizzle(TempVariableDeref->clone(ParseState, nullptr), 3, 0, 0, 0, 1),
+								new(ParseState) ir_expression(
+									ir_binop_div,
+									new(ParseState) ir_constant(1.0f),
+									new(ParseState) ir_swizzle(TempVariableDeref->clone(ParseState, nullptr), 3, 0, 0, 0, 1)
+									)
+								)
+							);
 
+						VariableDeref = TempVariableDeref->clone(ParseState, nullptr);
+					}
 					return VariableDeref;
 				}
 			}
@@ -1128,220 +1152,6 @@ static FSystemValue* SystemValueTable[] =
 	nullptr
 };
 
-/**
-* Generate an input semantic.
-* @param Frequency - The shader frequency.
-* @param ParseState - Parse state.
-* @param Semantic - The semantic name to generate.
-* @param Type - Value type.
-* @param DeclInstructions - IR to which declarations may be added.
-* @returns reference to IR variable for the semantic.
-*/
-static ir_rvalue* GenShaderInputSemantic(
-	EHlslShaderFrequency Frequency,
-	_mesa_glsl_parse_state* ParseState,
-	const char* Semantic,
-	const glsl_type* Type,
-	exec_list* DeclInstructions,
-	int SemanticArraySize,
-	int SemanticArrayIndex)
-{
-	ir_variable* Variable = NULL;
-	if (Semantic && FCStringAnsi::Strnicmp(Semantic, "SV_", 3) == 0)
-	{
-		FSystemValue* SystemValues = SystemValueTable[Frequency];
-		for (int i = 0; SystemValues[i].Semantic != NULL; ++i)
-		{
-			if (SystemValues[i].Mode == ir_var_in
-				//				&& (!SystemValues[i].bESOnly || ParseState->bGenerateES)
-				&& FCStringAnsi::Stricmp(SystemValues[i].Semantic, Semantic) == 0)
-			{
-				check(0);
-#if 0
-				if (SystemValues[i].bArrayVariable)
-				{
-					// Built-in array variable. Like gl_in[x].gl_Position.
-					// The variable for it has already been created in GenShaderInput().
-					/*ir_variable**/ Variable = ParseState->symbols->get_variable("gl_in");
-					check(Variable);
-					ir_dereference_variable* ArrayDeref = new(ParseState)ir_dereference_variable(Variable);
-					ir_dereference_array* StructDeref = new(ParseState)ir_dereference_array(
-						ArrayDeref,
-						new(ParseState)ir_constant((unsigned)SemanticArrayIndex)
-						);
-					ir_dereference_record* VariableDeref = new(ParseState)ir_dereference_record(
-						StructDeref,
-						SystemValues[i].GlslName
-						);
-					// TO DO - in case of SV_ClipDistance, we need to defer appropriate index in variable too.
-					return VariableDeref;
-				}
-				else
-				{
-					// Built-in variable that shows up only once, like gl_FragCoord in fragment
-					// shader, or gl_PrimitiveIDIn in geometry shader. Unlike gl_in[x].gl_Position.
-					// Even in geometry shader input pass it shows up only once.
-
-					// Create it on first pass, ignore the call on others.
-					if (SemanticArrayIndex == 0)
-					{
-						ir_variable* Variable = new(ParseState)ir_variable(
-							SystemValues[i].Type,
-							SystemValues[i].GlslName,
-							ir_var_in
-							);
-						Variable->read_only = true;
-						Variable->origin_upper_left = SystemValues[i].bOriginUpperLeft;
-						DeclInstructions->push_tail(Variable);
-						ParseState->symbols->add_variable(Variable);
-						ir_dereference_variable* VariableDeref = new(ParseState)ir_dereference_variable(Variable);
-
-						return VariableDeref;
-					}
-					else
-					{
-						return NULL;
-					}
-				}
-#endif
-			}
-		}
-	}
-
-	if (Variable)
-	{
-		// Up to this point, variables aren't contained in structs
-		DeclInstructions->push_tail(Variable);
-		ParseState->symbols->add_variable(Variable);
-		Variable->centroid = false;// InputQualifier.Fields.bCentroid;
-		Variable->interpolation = false;//InputQualifier.Fields.InterpolationMode;
-		Variable->is_patch_constant = false;//InputQualifier.Fields.bIsPatchConstant;
-		ir_rvalue* VariableDeref = new(ParseState)ir_dereference_variable(Variable);
-
-		return VariableDeref;
-	}
-
-	// If we're here, no built-in variables matched.
-
-	if (Semantic && FCStringAnsi::Strnicmp(Semantic, "SV_", 3) == 0)
-	{
-		_mesa_glsl_warning(ParseState, "unrecognized system "
-			"value input '%s'", Semantic);
-	}
-
-	if (Frequency == HSF_VertexShader || ParseState->bGenerateES)
-	{
-		const char* Prefix = "in";
-		if (ParseState->bGenerateES && Frequency == HSF_PixelShader)
-		{
-			Prefix = "var";
-		}
-
-		// Vertex shader inputs don't get packed into structs that we'll later morph into interface blocks
-		if (ParseState->bGenerateES && Type->is_integer())
-		{
-			// Convert integer attributes to floats
-			ir_variable* Variable = new(ParseState)ir_variable(
-				Type,
-				ralloc_asprintf(ParseState, "%s_%s_I", Prefix, Semantic),
-				ir_var_temporary
-				);
-			Variable->centroid = false;//InputQualifier.Fields.bCentroid;
-			Variable->interpolation = false;//InputQualifier.Fields.InterpolationMode;
-			check(Type->is_vector() || Type->is_scalar());
-			check(Type->base_type == GLSL_TYPE_INT || Type->base_type == GLSL_TYPE_UINT);
-
-			// New float attribute
-			ir_variable* ReplacedAttributeVar = new (ParseState)ir_variable(glsl_type::get_instance(GLSL_TYPE_FLOAT, Variable->type->vector_elements, 1), ralloc_asprintf(ParseState, "%s_%s", Prefix, Semantic), ir_var_in);
-			ReplacedAttributeVar->read_only = true;
-			ReplacedAttributeVar->centroid = false;//InputQualifier.Fields.bCentroid;
-			ReplacedAttributeVar->interpolation = false;//InputQualifier.Fields.InterpolationMode;
-
-			// Convert to integer
-			ir_assignment* ConversionAssignment = new(ParseState)ir_assignment(
-				new(ParseState)ir_dereference_variable(Variable),
-				new(ParseState)ir_expression(
-				Type->base_type == GLSL_TYPE_INT ? ir_unop_f2i : ir_unop_f2u,
-				new (ParseState)ir_dereference_variable(ReplacedAttributeVar)
-				)
-				);
-
-			DeclInstructions->push_tail(ReplacedAttributeVar);
-			DeclInstructions->push_tail(Variable);
-			DeclInstructions->push_tail(ConversionAssignment);
-			ParseState->symbols->add_variable(Variable);
-			ParseState->symbols->add_variable(ReplacedAttributeVar);
-
-			ir_dereference_variable* VariableDeref = new(ParseState)ir_dereference_variable(ReplacedAttributeVar);
-			return VariableDeref;
-		}
-
-		// Regular attribute
-		Variable = new(ParseState)ir_variable(
-			Type,
-			ralloc_asprintf(ParseState, "%s_%s", Prefix, Semantic),
-			ir_var_in
-			);
-		Variable->read_only = true;
-		Variable->centroid = false;//InputQualifier.Fields.bCentroid;
-		Variable->interpolation = false;//InputQualifier.Fields.InterpolationMode;
-		Variable->is_patch_constant = false;//InputQualifier.Fields.bIsPatchConstant;
-
-		DeclInstructions->push_tail(Variable);
-		ParseState->symbols->add_variable(Variable);
-
-		ir_dereference_variable* VariableDeref = new(ParseState)ir_dereference_variable(Variable);
-		return VariableDeref;
-	}
-	else if (SemanticArrayIndex == 0)
-	{
-		// On first pass, create variable
-
-		glsl_struct_field *StructField = ralloc_array(ParseState, glsl_struct_field, 1);
-
-		memset(StructField, 0, sizeof(glsl_struct_field));
-		StructField[0].type = Type;
-		StructField[0].name = ralloc_strdup(ParseState, "Data");
-
-		const glsl_type* VariableType = glsl_type::get_record_instance(StructField, 1, ralloc_strdup(ParseState, Semantic));
-		if (SemanticArraySize)
-		{
-			// Pack it into an array too
-			VariableType = glsl_type::get_array_instance(VariableType, SemanticArraySize);
-		}
-
-		ir_variable* Variable = new(ParseState)ir_variable(VariableType, ralloc_asprintf(ParseState, "in_%s", Semantic), ir_var_in);
-		Variable->read_only = true;
-		Variable->is_interface_block = true;
-		Variable->centroid = false;//InputQualifier.Fields.bCentroid;
-		Variable->interpolation = false;//InputQualifier.Fields.InterpolationMode;
-		Variable->is_patch_constant = false;//InputQualifier.Fields.bIsPatchConstant;
-		DeclInstructions->push_tail(Variable);
-		ParseState->symbols->add_variable(Variable);
-
-		ir_rvalue* VariableDeref = new(ParseState)ir_dereference_variable(Variable);
-		if (SemanticArraySize)
-		{
-			// Deref inside array first
-			VariableDeref = new(ParseState)ir_dereference_array(VariableDeref, new(ParseState)ir_constant((unsigned)SemanticArrayIndex)
-				);
-		}
-		VariableDeref = new(ParseState)ir_dereference_record(VariableDeref, ralloc_strdup(ParseState, "Data"));
-		return VariableDeref;
-	}
-	else
-	{
-		// Array variable, not first pass. It already exists, get it.
-		ir_variable* Variable = ParseState->symbols->get_variable(ralloc_asprintf(ParseState, "in_%s", Semantic));
-		check(Variable);
-
-		ir_rvalue* VariableDeref = new(ParseState)ir_dereference_variable(Variable);
-		VariableDeref = new(ParseState)ir_dereference_array(VariableDeref, new(ParseState)ir_constant((unsigned)SemanticArrayIndex));
-		VariableDeref = new(ParseState)ir_dereference_record(VariableDeref, ralloc_strdup(ParseState, "Data"));
-		return VariableDeref;
-	}
-}
-
 /**
 * Generate a shader input.
 * @param Frequency - The shader frequency.
