fix for crash when exiting the project settings while there is a long running IPP operation ongoing
UE-7166
#ios

diff --git a/Engine/Source/Developer/iOS/IOSPlatformEditor/Private/IOSTargetSettingsCustomization.cpp b/Engine/Source/Developer/iOS/IOSPlatformEditor/Private/IOSTargetSettingsCustomization.cpp
--- a/Engine/Source/Developer/iOS/IOSPlatformEditor/Private/IOSTargetSettingsCustomization.cpp
+++ b/Engine/Source/Developer/iOS/IOSPlatformEditor/Private/IOSTargetSettingsCustomization.cpp
@@ -68,6 +68,15 @@ FIOSTargetSettingsCustomization::FIOSTargetSettingsCustomization()
 	new (LaunchImageNames) FPlatformIconInfo(TEXT("Default-IPhone6Plus-Portrait.png"), LOCTEXT("LaunchImage_iPhone6Plus_Portrait", "Launch iPhone 6 Plus in Portrait"), FText::GetEmpty(), 1242, 2208, FPlatformIconInfo::Required);
 }
 
+FIOSTargetSettingsCustomization::~FIOSTargetSettingsCustomization()
+{
+	if (IPPProcess.IsValid())
+	{
+		IPPProcess = NULL;
+		FTicker::GetCoreTicker().RemoveTicker(TickerHandle);
+	}
+}
+
 void FIOSTargetSettingsCustomization::CustomizeDetails(IDetailLayoutBuilder& DetailLayout)
 {
 	SavedLayoutBuilder = &DetailLayout;
@@ -511,21 +520,6 @@ void FIOSTargetSettingsCustomization::OnPlistPropertyModified()
 	DeviceCapsArrayBody += TEXT("\t");
 	Updater.ReplaceKey(RequiredDeviceCaps, ClosingArray, DeviceCapsArrayBody);
 
-	// build the replacement device families
-	const FString DeviceFamilyKey(TEXT("<key>UIDeviceFamily</key>"));
-	FString FamilyKeyBody = TEXT("\n\t<array>\n");
-	// automatically add armv7 for now
-	if (Settings.bSupportsIPhone)
-	{
-		FamilyKeyBody += TEXT("\t\t<integer>1</integer>\n");
-	}
-	if (Settings.bSupportsIPad)
-	{
-		FamilyKeyBody += TEXT("\t\t<integer>2</integer>\n");
-	}
-	FamilyKeyBody += TEXT("\t");
-	Updater.ReplaceKey(DeviceFamilyKey, ClosingArray, FamilyKeyBody);
-
 	// build the replacement min iOS version
 	const FString MiniOSVersionKey(TEXT("<key>MinimumOSVersion</key>"));
 	FString iOSVersionBody = TEXT("\n\t<string>");
@@ -662,7 +656,7 @@ FReply FIOSTargetSettingsCustomization::OnInstallProvisionClicked()
 		OutputMessage = TEXT("");
 		IPPProcess->OnOutput().BindStatic(&OnOutput);
 		IPPProcess->Launch();
-		FTicker::GetCoreTicker().AddTicker(FTickerDelegate::CreateRaw(this, &FIOSTargetSettingsCustomization::UpdateStatusDelegate), 10.0f);
+		TickerHandle = FTicker::GetCoreTicker().AddTicker(FTickerDelegate::CreateRaw(this, &FIOSTargetSettingsCustomization::UpdateStatusDelegate), 10.0f);
 	}
 
 	return FReply::Handled();
@@ -720,7 +714,7 @@ FReply FIOSTargetSettingsCustomization::OnInstallCertificateClicked()
 		OutputMessage = TEXT("");
 		IPPProcess->OnOutput().BindStatic(&OnOutput);
 		IPPProcess->Launch();
-		FTicker::GetCoreTicker().AddTicker(FTickerDelegate::CreateRaw(this, &FIOSTargetSettingsCustomization::UpdateStatusDelegate), 10.0f);
+		TickerHandle = FTicker::GetCoreTicker().AddTicker(FTickerDelegate::CreateRaw(this, &FIOSTargetSettingsCustomization::UpdateStatusDelegate), 10.0f);
 	}
 
 	return FReply::Handled();
@@ -764,7 +758,6 @@ bool FIOSTargetSettingsCustomization::UpdateStatusDelegate(float DeltaTime)
 	}
 	int RetCode = IPPProcess->GetReturnCode();
 	IPPProcess = NULL;
-	ensure(RetCode == 0);
 	UpdateStatus();
 
 	return false;
diff --git a/Engine/Source/Developer/iOS/IOSPlatformEditor/Private/IOSTargetSettingsCustomization.h b/Engine/Source/Developer/iOS/IOSPlatformEditor/Private/IOSTargetSettingsCustomization.h
--- a/Engine/Source/Developer/iOS/IOSPlatformEditor/Private/IOSTargetSettingsCustomization.h
+++ b/Engine/Source/Developer/iOS/IOSPlatformEditor/Private/IOSTargetSettingsCustomization.h
@@ -12,6 +12,8 @@
 class FIOSTargetSettingsCustomization : public IDetailCustomization
 {
 public:
+	~FIOSTargetSettingsCustomization();
+
 	// Makes a new instance of this detail layout class for a specific detail view requesting it
 	static TSharedRef<IDetailCustomization> MakeInstance();
 
@@ -38,6 +40,7 @@ private:
 	bool bCertificateInstalled;
 
 	TSharedPtr<FMonitoredProcess> IPPProcess;
+	FDelegateHandle TickerHandle;
 
 private:
 	FIOSTargetSettingsCustomization();
