UE4 - Metal - Fix RWTexture2D writes (still need to check if we can bind textures as buffer[[n]] )

diff --git a/Engine/Source/Developer/iOS/MetalShaderFormat/Private/MetalBackend.cpp b/Engine/Source/Developer/iOS/MetalShaderFormat/Private/MetalBackend.cpp
--- a/Engine/Source/Developer/iOS/MetalShaderFormat/Private/MetalBackend.cpp
+++ b/Engine/Source/Developer/iOS/MetalShaderFormat/Private/MetalBackend.cpp
@@ -651,17 +651,34 @@ protected:
 				// Buffer
 				int BufferIndex = Buffers.GetIndex(var);
 				check(BufferIndex >= 0);
-				ralloc_asprintf_append(
-					buffer,
-					"device "
-					);
-				print_type_pre(PtrType->inner_type);
-				ralloc_asprintf_append(buffer, " *%s", unique_name(var));
-				print_type_post(PtrType->inner_type);
-				ralloc_asprintf_append(
-					buffer,
-					" [[ buffer(%d) ]]", BufferIndex
-					);
+				if (var->type->sampler_buffer)
+				{
+					ralloc_asprintf_append(
+						buffer,
+						"device "
+						);
+					print_type_pre(PtrType->inner_type);
+					ralloc_asprintf_append(buffer, " *%s", unique_name(var));
+					print_type_post(PtrType->inner_type);
+					ralloc_asprintf_append(
+						buffer,
+						" [[ buffer(%d) ]]", BufferIndex
+						);
+				}
+				else
+				{
+					check(PtrType->inner_type->is_numeric());
+					auto* Found = strstr(PtrType->name, "image");
+					check(Found);
+					Found += 5;	//strlen(image)
+					ralloc_asprintf_append(buffer, "texture%s<", Found);
+					print_type_pre(PtrType->inner_type);
+					ralloc_asprintf_append(buffer, ", access::write> %s", unique_name(var));
+					ralloc_asprintf_append(
+						buffer,
+						" [[ buffer(%d) ]]", BufferIndex
+						);
+				}
 			}
 			else
 			{
@@ -1225,21 +1242,43 @@ protected:
 
 		if ( deref->op == ir_image_access)
 		{
-			if ( src == NULL )
+			bool bIsRWTexture = !deref->image->type->sampler_buffer;
+			if (src == nullptr)
 			{
 				deref->image->accept(this);
-				ralloc_asprintf_append( buffer, "[" );
-				deref->image_index->accept(this);
-				ralloc_asprintf_append(buffer, "]"/*.%s, swizzle[dst_elements - 1]*/);
+				if (bIsRWTexture)
+				{
+					ralloc_asprintf_append(buffer, ".read(");
+					deref->image_index->accept(this);
+					ralloc_asprintf_append(buffer, ")");
+				}
+				else
+				{
+					ralloc_asprintf_append(buffer, "[");
+					deref->image_index->accept(this);
+					ralloc_asprintf_append(buffer, "]"/*.%s, swizzle[dst_elements - 1]*/);
+				}
 			}
 			else
 			{
 				deref->image->accept(this);
-				ralloc_asprintf_append( buffer, "[" );
-				deref->image_index->accept(this);
-				ralloc_asprintf_append( buffer, "] = " );
-				src->accept(this);
-				ralloc_asprintf_append(buffer, ""/*".%s", expand[src_elements - 1]*/);
+				if (bIsRWTexture)
+				{
+					ralloc_asprintf_append(buffer, ".write(");
+					src->accept(this);
+					ralloc_asprintf_append(buffer, ",");
+					deref->image_index->accept(this);
+					ralloc_asprintf_append(buffer, ")");
+				}
+				else
+				{
+					deref->image->accept(this);
+					ralloc_asprintf_append( buffer, "[" );
+					deref->image_index->accept(this);
+					ralloc_asprintf_append( buffer, "] = " );
+					src->accept(this);
+					ralloc_asprintf_append(buffer, ""/*".%s", expand[src_elements - 1]*/);
+				}
 			}
 		}
 		else if ( deref->op == ir_image_dimensions)
