UE4 - Metal - Fix: texture2D into texture2d

diff --git a/Engine/Source/Developer/iOS/MetalShaderFormat/Private/MetalBackend.cpp b/Engine/Source/Developer/iOS/MetalShaderFormat/Private/MetalBackend.cpp
--- a/Engine/Source/Developer/iOS/MetalShaderFormat/Private/MetalBackend.cpp
+++ b/Engine/Source/Developer/iOS/MetalShaderFormat/Private/MetalBackend.cpp
@@ -670,8 +670,17 @@ protected:
 					check(PtrType->inner_type->is_numeric());
 					auto* Found = strstr(PtrType->name, "image");
 					check(Found);
-					Found += 5;	//strlen(image)
-					ralloc_asprintf_append(buffer, "texture%s<", Found);
+					Found += 5;	//strlen("image")
+					char Temp[16];
+					char* TempPtr = Temp;
+					do
+					{
+						*TempPtr = toupper(*Found);
+						++Found;
+					}
+					while (*Found);
+					*TempPtr = 0;
+					ralloc_asprintf_append(buffer, "texture%s<", Temp);
 					print_type_pre(PtrType->inner_type);
 					ralloc_asprintf_append(buffer, ", access::write> %s", unique_name(var));
 					ralloc_asprintf_append(
