- Some compute shader fixes, one left with buffer indices (handing back to Rolando)

diff --git a/Engine/Source/Developer/iOS/MetalShaderFormat/Private/MetalBackend.cpp b/Engine/Source/Developer/iOS/MetalShaderFormat/Private/MetalBackend.cpp
--- a/Engine/Source/Developer/iOS/MetalShaderFormat/Private/MetalBackend.cpp
+++ b/Engine/Source/Developer/iOS/MetalShaderFormat/Private/MetalBackend.cpp
@@ -675,7 +675,7 @@ protected:
 					char* TempPtr = Temp;
 					do
 					{
-						*TempPtr = toupper(*Found);
+						*TempPtr++ = tolower(*Found);
 						++Found;
 					}
 					while (*Found);
