addition of certificate and provision viewer to project settings to make it more clear to the user why a provision or certificate is invalid
UE-8074
#ios

diff --git a/Engine/Source/Developer/iOS/IOSPlatformEditor/Private/IOSTargetSettingsCustomization.cpp b/Engine/Source/Developer/iOS/IOSPlatformEditor/Private/IOSTargetSettingsCustomization.cpp
--- a/Engine/Source/Developer/iOS/IOSPlatformEditor/Private/IOSTargetSettingsCustomization.cpp
+++ b/Engine/Source/Developer/iOS/IOSPlatformEditor/Private/IOSTargetSettingsCustomization.cpp
@@ -20,9 +20,15 @@
 #include "NotificationManager.h"
 #include "TargetPlatform.h"
 #include "GameProjectGenerationModule.h"
+#include "SHyperlink.h"
+#include "SProvisionListRow.h"
+#include "SCertificateListRow.h"
 
 #define LOCTEXT_NAMESPACE "IOSTargetSettings"
 
+bool SProvisionListRow::bInitialized = false;
+FCheckBoxStyle SProvisionListRow::ProvisionCheckBoxStyle;
+
 //////////////////////////////////////////////////////////////////////////
 // FIOSTargetSettingsCustomization
 namespace FIOSTargetSettingsCustomizationConstants
@@ -66,6 +72,9 @@ FIOSTargetSettingsCustomization::FIOSTargetSettingsCustomization()
 	new (LaunchImageNames) FPlatformIconInfo(TEXT("Default-IPhone6.png"), LOCTEXT("LaunchImage_iPhone6", "Launch iPhone 6"), FText::GetEmpty(), 750, 1334, FPlatformIconInfo::Required);
 	new (LaunchImageNames) FPlatformIconInfo(TEXT("Default-IPhone6Plus-Landscape.png"), LOCTEXT("LaunchImage_iPhone6Plus_Landscape", "Launch iPhone 6 Plus in Landscape"), FText::GetEmpty(), 2208, 1242, FPlatformIconInfo::Required);
 	new (LaunchImageNames) FPlatformIconInfo(TEXT("Default-IPhone6Plus-Portrait.png"), LOCTEXT("LaunchImage_iPhone6Plus_Portrait", "Launch iPhone 6 Plus in Portrait"), FText::GetEmpty(), 1242, 2208, FPlatformIconInfo::Required);
+
+	bShowAllProvisions = false;
+	bShowAllCertificates = false;
 }
 
 FIOSTargetSettingsCustomization::~FIOSTargetSettingsCustomization()
@@ -81,7 +90,7 @@ void FIOSTargetSettingsCustomization::CustomizeDetails(IDetailLayoutBuilder& Det
 {
 	SavedLayoutBuilder = &DetailLayout;
 
-	UpdateStatus();
+	FindRequiredFiles();
 
 	BuildPListSection(DetailLayout);
 
@@ -90,18 +99,105 @@ void FIOSTargetSettingsCustomization::CustomizeDetails(IDetailLayoutBuilder& Det
 	BuildRemoteBuildingSection(DetailLayout);
 }
 
+static FString OutputMessage;
+static void OnOutput(FString Message)
+{
+	OutputMessage += Message;
+	OutputMessage += "\n";
+	UE_LOG(LogTemp, Display, TEXT("%s\n"), *Message);
+}
+
 void FIOSTargetSettingsCustomization::UpdateStatus()
 {
-	// get the provision status
-	bProvisionInstalled = bCertificateInstalled = false;
-	const ITargetPlatform* const Platform = GetTargetPlatformManager()->FindTargetPlatform("IOS");
-	if (Platform)
+	if (OutputMessage.Len() > 0)
 	{
-		FString NotInstalledTutorialLink;
-		FString ProjectPath = FPaths::IsProjectFilePathSet() ? FPaths::ConvertRelativePathToFull(FPaths::GetProjectFilePath()) : FPaths::RootDir() / FApp::GetGameName() / FApp::GetGameName() + TEXT(".uproject");
-		int32 Result = Platform->CheckRequirements(ProjectPath, false, NotInstalledTutorialLink);
-		bProvisionInstalled = !(Result & ETargetPlatformReadyStatus::ProvisionNotFound);
-		bCertificateInstalled = !(Result & ETargetPlatformReadyStatus::SigningKeyNotFound);
+		CertificateList.Reset();
+		ProvisionList.Reset();
+
+		// Now split up the log into multiple lines
+		TArray<FString> LogLines;
+		OutputMessage.ParseIntoArray(&LogLines, LINE_TERMINATOR, true);
+		
+		// format of the line being read here!!
+		bool bCerts = false;
+		for (int Index = 0; Index < LogLines.Num(); Index++)
+		{
+			FString& Line = LogLines[Index];
+			TArray<FString> Fields;
+			Line.ParseIntoArray(&Fields, TEXT(","), true);
+			if (Line.Contains(TEXT("CERTIFICATE-"), ESearchCase::CaseSensitive))
+			{
+				CertificatePtr Cert = MakeShareable<FCertificate>(new FCertificate());
+				for (int FieldIndex = 0; FieldIndex < Fields.Num(); ++FieldIndex)
+				{
+					FString Key, Value;
+					Fields[FieldIndex].Split(TEXT(":"), &Key, &Value);
+					if (Key.Contains("Name"))
+					{
+						Cert->Name = Value;
+					}
+					else if (Key.Contains(TEXT("Validity")))
+					{
+						Cert->Status = Value;
+					}
+					else if (Key.Contains(TEXT("EndDate")))
+					{
+						FString Date, Time;
+						Value.Split(TEXT(" "), &Date, &Time);
+						Cert->Expires = Date;
+					}
+				}
+				CertificateList.Add(Cert);
+			}
+			else if (Line.Contains(TEXT("PROVISION-"), ESearchCase::CaseSensitive))
+			{
+				ProvisionPtr Prov = MakeShareable<FProvision>(new FProvision());
+				for (int FieldIndex = 0; FieldIndex < Fields.Num(); ++FieldIndex)
+				{
+					FString Key, Value;
+					Fields[FieldIndex].Split(TEXT(":"), &Key, &Value);
+					if (Key.Contains("File"))
+					{
+						Prov->FileName = Value;
+					}
+					else if (Key.Contains("Name"))
+					{
+						Prov->Name = Value;
+					}
+					else if (Key.Contains(TEXT("Validity")))
+					{
+						Prov->Status = Value;
+					}
+					else if (Key.Contains(TEXT("Type")))
+					{
+						Prov->bDistribution = Value.Contains(TEXT("DISTRIBUTION"));
+					}
+				}
+				ProvisionList.Add(Prov);
+			}
+			else if (Line.Contains(TEXT("MATCHED-"), ESearchCase::CaseSensitive))
+			{
+				for (int FieldIndex = 0; FieldIndex < Fields.Num(); ++FieldIndex)
+				{
+					FString Key, Value;
+					Fields[FieldIndex].Split(TEXT(":"), &Key, &Value);
+					if (Key.Contains("File"))
+					{
+						SelectedFile = Value;
+					}
+					else if (Key.Contains("Provision"))
+					{
+						SelectedProvision = Value;
+					}
+					else if (Key.Contains(TEXT("Cert")))
+					{
+						SelectedCert = Value;
+					}
+				}
+			}
+		}
+
+		FilterLists();
 	}
 }
 
@@ -122,8 +218,6 @@ void FIOSTargetSettingsCustomization::BuildPListSection(IDetailLayoutBuilder& De
 		.PlatformName(LOCTEXT("iOSPlatformName", "iOS"))
 		.OnSetupClicked(this, &FIOSTargetSettingsCustomization::CopySetupFilesIntoProject);
 
-	SetupForPlatformAttribute = PlatformSetupMessage->GetReadyToGoAttribute();
-
 /*	ProvisionCategory.AddCustomRow(TEXT("Certificate Request"), false)
 		.NameContent()
 		[
@@ -153,83 +247,209 @@ void FIOSTargetSettingsCustomization::BuildPListSection(IDetailLayoutBuilder& De
 		];*/
 
 	ProvisionCategory.AddCustomRow(LOCTEXT("ProvisionLabel", "Provision"), false)
-		.NameContent()
+		.WholeRowWidget
+		.MinDesiredWidth(0.f)
+		.MaxDesiredWidth(0.f)
+		.HAlign(HAlign_Fill)
 		[
-			SNew(SHorizontalBox)
-			+ SHorizontalBox::Slot()
-			.Padding(FMargin(0, 1, 0, 1))
-			.FillWidth(1.0f)
+			SNew(SVerticalBox)
+			+ SVerticalBox::Slot()
+			.Padding(FMargin(10, 10, 10, 10))
+			.AutoHeight()
 			[
-				SNew(STextBlock)
-				.Text(LOCTEXT("ProvisionLabel", "Provision"))
-				.Font(DetailLayout.GetDetailFont())
+				SAssignNew(ProvisionListView, SListView<ProvisionPtr>)
+				.ItemHeight(20.0f)
+				.ListItemsSource(&FilteredProvisionList)
+				.OnGenerateRow(this, &FIOSTargetSettingsCustomization::HandleProvisionListGenerateRow)
+				.SelectionMode(ESelectionMode::None)
+				.HeaderRow
+				(
+				SNew(SHeaderRow)
+				+ SHeaderRow::Column("Name")
+				.DefaultLabel(LOCTEXT("ProvisionListNameColumnHeader", "Provision"))
+				.FillWidth(1.0f)
+				+ SHeaderRow::Column("File")
+				.DefaultLabel(LOCTEXT("ProvisionListFileColumnHeader", "File"))
+				+ SHeaderRow::Column("Status")
+				.DefaultLabel(LOCTEXT("ProvisionListStatusColumnHeader", "Status"))
+				+ SHeaderRow::Column("Distribution")
+				.DefaultLabel(LOCTEXT("ProvisionListDistributionColumnHeader", "Distribution"))
+				.FixedWidth(75.0f)
+				)
 			]
-		]
-		.ValueContent()
-			.MinDesiredWidth(300.0f)
-			.MaxDesiredWidth(350.0f)
-		[
-			SNew(SHorizontalBox)
-			+ SHorizontalBox::Slot()
-			.Padding(FMargin(8,3,8,3))
-			.AutoWidth()
+			+ SVerticalBox::Slot()
+			.AutoHeight()
+			.Padding(0.0f, 6.0f, 0.0f, 4.0f)
 			[
-				SNew(SImage)
-				.Image(this, &FIOSTargetSettingsCustomization::GetProvisionStatus)
+				SNew(SSeparator)
+				.Orientation(Orient_Horizontal)
 			]
-			+ SHorizontalBox::Slot()
-			.Padding(FMargin(0, 0, 0, 0))
-			.AutoWidth()
+			+ SVerticalBox::Slot()
+			.AutoHeight()
 			[
-				SNew(SButton)
-				.HAlign(HAlign_Center)
-				.VAlign(VAlign_Center)
-				.OnClicked(this, &FIOSTargetSettingsCustomization::OnInstallProvisionClicked)
+				SNew(SHorizontalBox)
+
+				+ SHorizontalBox::Slot()
+				.AutoWidth()
+				[
+					SNew(SRichTextBlock)
+					.Text(LOCTEXT("ProvisionMessage", "<RichTextBlock.TextHighlight>Note</>: The provision in green will be used to provision the IPA."))
+					.TextStyle(FEditorStyle::Get(), "MessageLog")
+					.DecoratorStyleSet(&FEditorStyle::Get())
+					.AutoWrapText(true)
+				]
+
+				+ SHorizontalBox::Slot()
+				.FillWidth(1.0f)
+				.HAlign(HAlign_Right)
 				[
 					SNew(STextBlock)
-					.Text(FText::FromString("Import Provision"))
+					.Text(LOCTEXT("ViewLabel", "View:"))
 				]
+
+				+ SHorizontalBox::Slot()
+					.AutoWidth()
+					.Padding(8.0f, 0.0f)
+					[
+						// all provisions hyper link
+						SNew(SHyperlink)
+						.OnNavigate(this, &FIOSTargetSettingsCustomization::HandleAllProvisionsHyperlinkNavigate, true)
+						.Text(LOCTEXT("AllProvisionsHyperLinkLabel", "All"))
+						.ToolTipText(LOCTEXT("AllProvisionsButtonTooltip", "View all provisions."))
+					]
+
+				+ SHorizontalBox::Slot()
+					.AutoWidth()
+					[
+						// valid provisions hyper link
+						SNew(SHyperlink)
+						.OnNavigate(this, &FIOSTargetSettingsCustomization::HandleAllProvisionsHyperlinkNavigate, false)
+						.Text(LOCTEXT("ValidProvisionsHyperlinkLabel", "Valid Only"))
+						.ToolTipText(LOCTEXT("ValidProvisionsHyperlinkTooltip", "View Valid provisions."))
+					]
+			]
+			+ SVerticalBox::Slot()
+			.AutoHeight()
+			[
+				SNew(SHorizontalBox)
+				+ SHorizontalBox::Slot()
+					.Padding(FMargin(0, 5, 0, 10))
+					.AutoWidth()
+					[
+						SNew(SButton)
+						.HAlign(HAlign_Center)
+						.VAlign(VAlign_Center)
+						.OnClicked(this, &FIOSTargetSettingsCustomization::OnInstallProvisionClicked)
+						[
+							SNew(STextBlock)
+							.Text(FText::FromString("Import Provision"))
+						]
+					]
 			]
 		];
 
 	ProvisionCategory.AddCustomRow(LOCTEXT("CertificateLabel", "Certificate"), false)
-		.NameContent()
+		.WholeRowWidget
+		.MinDesiredWidth(0.f)
+		.MaxDesiredWidth(0.f)
+		.HAlign(HAlign_Fill)
 		[
-			SNew(SHorizontalBox)
-			+ SHorizontalBox::Slot()
-			.Padding(FMargin(0, 1, 0, 1))
-			.FillWidth(1.0f)
-			[
-				SNew(STextBlock)
-				.Text(LOCTEXT("CertificateLabel", "Certificate"))
-				.Font(DetailLayout.GetDetailFont())
-			]
-		]
-		.ValueContent()
-			.MinDesiredWidth(300.0f)
-			.MaxDesiredWidth(350.0f)
+			SNew(SVerticalBox)
+			+ SVerticalBox::Slot()
+			.AutoHeight()
 			[
 				SNew(SHorizontalBox)
 				+ SHorizontalBox::Slot()
-				.AutoWidth()
-				.Padding(FMargin(8,3,8,3))
+				.Padding(FMargin(10, 10, 10, 10))
+				.FillWidth(1.0f)
 				[
-					SNew(SImage)
-					.Image(this, &FIOSTargetSettingsCustomization::GetCertificateStatus)
+					SAssignNew(CertificateListView, SListView<CertificatePtr>)
+					.ItemHeight(20.0f)
+					.ListItemsSource(&FilteredCertificateList)
+					.OnGenerateRow(this, &FIOSTargetSettingsCustomization::HandleCertificateListGenerateRow)
+					.SelectionMode(ESelectionMode::None)
+					.HeaderRow
+					(
+						SNew(SHeaderRow)
+						+ SHeaderRow::Column("Name")
+						.DefaultLabel(LOCTEXT("CertificateListNameColumnHeader", "Certificate"))
+						+ SHeaderRow::Column("Status")
+						.DefaultLabel(LOCTEXT("CertificateListStatusColumnHeader", "Status"))
+						.FixedWidth(75.0f)
+						+ SHeaderRow::Column("Expires")
+						.DefaultLabel(LOCTEXT("CertificateListExpiresColumnHeader", "Expires"))
+						.FixedWidth(75.0f)
+					)
 				]
-				+ SHorizontalBox::Slot()
+			]
+			+ SVerticalBox::Slot()
+				.AutoHeight()
+				.Padding(0.0f, 6.0f, 0.0f, 4.0f)
+				[
+					SNew(SSeparator)
+					.Orientation(Orient_Horizontal)
+				]
+			+ SVerticalBox::Slot()
+				.AutoHeight()
+				[
+					SNew(SHorizontalBox)
+					+ SHorizontalBox::Slot()
 					.AutoWidth()
 					[
-						SNew(SButton)
-						.HAlign(HAlign_Center)
-						.VAlign(VAlign_Center)
-						.OnClicked(this, &FIOSTargetSettingsCustomization::OnInstallCertificateClicked)
+						SNew(SRichTextBlock)
+						.Text(LOCTEXT("CertificateMessage", "<RichTextBlock.TextHighlight>Note</>: The certificate in green will be used to sign the IPA."))
+						.TextStyle(FEditorStyle::Get(), "MessageLog")
+						.DecoratorStyleSet(&FEditorStyle::Get())
+						.AutoWrapText(true)
+					]
+					+ SHorizontalBox::Slot()
+					.FillWidth(1.0f)
+					.HAlign(HAlign_Right)
+					[
+						SNew(STextBlock)
+						.Text(LOCTEXT("ViewLabel", "View:"))
+					]
+
+					+ SHorizontalBox::Slot()
+						.AutoWidth()
+						.Padding(8.0f, 0.0f)
 						[
-							SNew(STextBlock)
-							.Text(FText::FromString("Import Certificate"))
+							// all provisions hyper link
+							SNew(SHyperlink)
+							.OnNavigate(this, &FIOSTargetSettingsCustomization::HandleAllCertificatesHyperlinkNavigate, true)
+							.Text(LOCTEXT("AllCertificatesHyperLinkLabel", "All"))
+							.ToolTipText(LOCTEXT("AllCertificatesButtonTooltip", "View all certificates."))
 						]
+
+					+ SHorizontalBox::Slot()
+						.AutoWidth()
+						[
+							// valid provisions hyper link
+							SNew(SHyperlink)
+							.OnNavigate(this, &FIOSTargetSettingsCustomization::HandleAllCertificatesHyperlinkNavigate, false)
+							.Text(LOCTEXT("ValidCertificatesHyperlinkLabel", "Valid Only"))
+							.ToolTipText(LOCTEXT("ValidCertificatesHyperlinkTooltip", "View Valid certificates."))
+						]
+				]
+			+ SVerticalBox::Slot()
+			.AutoHeight()
+			[
+				SNew(SHorizontalBox)
+				+ SHorizontalBox::Slot()
+				.Padding(FMargin(0, 5, 0, 10))
+				.AutoWidth()
+				[
+					SNew(SButton)
+					.HAlign(HAlign_Center)
+					.VAlign(VAlign_Center)
+					.OnClicked(this, &FIOSTargetSettingsCustomization::OnInstallCertificateClicked)
+					[
+						SNew(STextBlock)
+						.Text(FText::FromString("Import Certificate"))
 					]
-			];
+				]
+			]
+		];
 	
 	BundleCategory.AddCustomRow(LOCTEXT("UpgradeInfo", "Upgrade Info"), false)
 	.WholeRowWidget
@@ -253,7 +473,7 @@ void FIOSTargetSettingsCustomization::BuildPListSection(IDetailLayoutBuilder& De
 	 ];
 
 	// Show properties that are gated by the plist being present and writable
-	FSimpleDelegate PlistModifiedDelegate = FSimpleDelegate::CreateRaw(this, &FIOSTargetSettingsCustomization::OnPlistPropertyModified);
+	FSimpleDelegate BundleModifiedDelegate = FSimpleDelegate::CreateRaw(this, &FIOSTargetSettingsCustomization::FindRequiredFiles);
 	FGameProjectGenerationModule& GameProjectModule = FModuleManager::LoadModuleChecked<FGameProjectGenerationModule>(TEXT("GameProjectGeneration"));
 	bool bHasCode = GameProjectModule.Get().ProjectHasCodeFiles();
 
@@ -261,7 +481,6 @@ void FIOSTargetSettingsCustomization::BuildPListSection(IDetailLayoutBuilder& De
 	{ \
 		TSharedRef<IPropertyHandle> PropertyHandle = DetailLayout.GetProperty(GET_MEMBER_NAME_CHECKED(UIOSRuntimeSettings, PropName)); \
 		Category.AddProperty(PropertyHandle) \
-			.EditCondition(SetupForPlatformAttribute, NULL) \
 			.IsEnabled(!FRocketSupport::IsRocket()) \
 			.ToolTip(!FRocketSupport::IsRocket() ? Tip : DisabledTip); \
 	}
@@ -269,15 +488,21 @@ void FIOSTargetSettingsCustomization::BuildPListSection(IDetailLayoutBuilder& De
 #define SETUP_PLIST_PROP(PropName, Category, Tip) \
 	{ \
 		TSharedRef<IPropertyHandle> PropertyHandle = DetailLayout.GetProperty(GET_MEMBER_NAME_CHECKED(UIOSRuntimeSettings, PropName)); \
-		PropertyHandle->SetOnPropertyValueChanged(PlistModifiedDelegate); \
 		Category.AddProperty(PropertyHandle) \
-			.EditCondition(SetupForPlatformAttribute, NULL) \
 			.ToolTip(Tip); \
 	}
 
+#define SETUP_STATUS_PROP(PropName, Category, Tip) \
+	{ \
+	TSharedRef<IPropertyHandle> PropertyHandle = DetailLayout.GetProperty(GET_MEMBER_NAME_CHECKED(UIOSRuntimeSettings, PropName)); \
+	PropertyHandle->SetOnPropertyValueChanged(BundleModifiedDelegate); \
+	Category.AddProperty(PropertyHandle) \
+		.ToolTip(Tip); \
+	}
+
 	SETUP_PLIST_PROP(BundleDisplayName, BundleCategory, LOCTEXT("BundleDisplayNameToolTip", "Specifies the the display name for the application. This will be displayed under the icon on the device."));
 	SETUP_PLIST_PROP(BundleName, BundleCategory, LOCTEXT("BundleNameToolTip", "Specifies the the name of the application bundle. This is the short name for the application bundle."));
-	SETUP_PLIST_PROP(BundleIdentifier, BundleCategory, LOCTEXT("BundleIdentifierToolTip", "Specifies the bundle identifier for the application."));
+	SETUP_STATUS_PROP(BundleIdentifier, BundleCategory, LOCTEXT("BundleIdentifierToolTip", "Specifies the bundle identifier for the application."));
 	SETUP_PLIST_PROP(VersionInfo, BundleCategory, LOCTEXT("VersionInfoToolTip", "Specifies the version for the application."));
 	SETUP_PLIST_PROP(bSupportsPortraitOrientation, OrientationCategory, LOCTEXT("SupportsPortraitOrientationToolTip", "Supports default portrait orientation. Landscape will not be supported."));
 	SETUP_PLIST_PROP(bSupportsUpsideDownOrientation, OrientationCategory, LOCTEXT("SupportsUpsideDownOrientationToolTip", "Supports upside down portrait orientation. Landscape will not be supported."));
@@ -318,7 +543,6 @@ void FIOSTargetSettingsCustomization::BuildRemoteBuildingSection(IDetailLayoutBu
 	TSharedRef<IPropertyHandle> RemoteServerNamePropertyHandle = DetailLayout.GetProperty(GET_MEMBER_NAME_CHECKED(UIOSRuntimeSettings, RemoteServerName));
 	IDetailPropertyRow& RemoteServerNamePropertyRow = RemoteBuildingGroup.AddPropertyRow(RemoteServerNamePropertyHandle);
 	RemoteServerNamePropertyRow
-		.EditCondition(SetupForPlatformAttribute, NULL)
 		.IsEnabled(!FRocketSupport::IsRocket())
 		.ToolTip(!FRocketSupport::IsRocket() ? LOCTEXT("RemoteServerNameToolTip", "The name or ip address of the remote mac which will be used to build IOS") : FIOSTargetSettingsCustomizationConstants::DisabledTip);
 
@@ -327,7 +551,6 @@ void FIOSTargetSettingsCustomization::BuildRemoteBuildingSection(IDetailLayoutBu
 	TSharedRef<IPropertyHandle> UseRSyncPropertyHandle = DetailLayout.GetProperty(GET_MEMBER_NAME_CHECKED(UIOSRuntimeSettings, bUseRSync));
 	IDetailPropertyRow& UseRSyncPropertRow = RemoteBuildingGroup.AddPropertyRow(UseRSyncPropertyHandle);
 	UseRSyncPropertRow
-		.EditCondition(SetupForPlatformAttribute, NULL)
 		.IsEnabled(!FRocketSupport::IsRocket())
 		.ToolTip(!FRocketSupport::IsRocket() ? LOCTEXT("UseRSyncToolTip", "Use RSync instead of RPCUtility") : FIOSTargetSettingsCustomizationConstants::DisabledTip);
 
@@ -336,7 +559,6 @@ void FIOSTargetSettingsCustomization::BuildRemoteBuildingSection(IDetailLayoutBu
 	TSharedRef<IPropertyHandle> RSyncUsernamePropertyHandle = DetailLayout.GetProperty(GET_MEMBER_NAME_CHECKED(UIOSRuntimeSettings, RSyncUsername));
 	IDetailPropertyRow& RSyncUsernamePropertyRow = RemoteBuildingGroup.AddPropertyRow(RSyncUsernamePropertyHandle);
 	RSyncUsernamePropertyRow
-		.EditCondition(SetupForPlatformAttribute, NULL)
 		.IsEnabled(!FRocketSupport::IsRocket())
 		.ToolTip(!FRocketSupport::IsRocket() ? LOCTEXT("RSyncUsernameToolTip", "The username of the mac user that matches the specified SSH Key.") : FIOSTargetSettingsCustomizationConstants::DisabledTip);
 
@@ -345,7 +567,6 @@ void FIOSTargetSettingsCustomization::BuildRemoteBuildingSection(IDetailLayoutBu
 	TSharedRef<IPropertyHandle> SSHPrivateKeyLocationPropertyHandle = DetailLayout.GetProperty(GET_MEMBER_NAME_CHECKED(UIOSRuntimeSettings, SSHPrivateKeyLocation));
 	IDetailPropertyRow& SSHPrivateKeyLocationPropertyRow = RemoteBuildingGroup.AddPropertyRow(SSHPrivateKeyLocationPropertyHandle);
 	SSHPrivateKeyLocationPropertyRow
-		.EditCondition(SetupForPlatformAttribute, NULL)
 		.IsEnabled(!FRocketSupport::IsRocket())
 		.ToolTip(!FRocketSupport::IsRocket() ? LOCTEXT("SSHPrivateKeyLocationToolTip", "The existing location of an SSH Key found by UE4.") : FIOSTargetSettingsCustomizationConstants::DisabledTip);
 
@@ -354,7 +575,6 @@ void FIOSTargetSettingsCustomization::BuildRemoteBuildingSection(IDetailLayoutBu
 	TSharedRef<IPropertyHandle> SSHPrivateKeyOverridePathPropertyHandle = DetailLayout.GetProperty(GET_MEMBER_NAME_CHECKED(UIOSRuntimeSettings, SSHPrivateKeyOverridePath));
 	IDetailPropertyRow& SSHPrivateKeyOverridePathPropertyRow = RemoteBuildingGroup.AddPropertyRow(SSHPrivateKeyOverridePathPropertyHandle);
 	SSHPrivateKeyOverridePathPropertyRow
-		.EditCondition(SetupForPlatformAttribute, NULL)
 		.IsEnabled(!FRocketSupport::IsRocket())
 		.ToolTip(!FRocketSupport::IsRocket() ? LOCTEXT("SSHPrivateKeyOverridePathToolTip", "Override the existing SSH Private Key with one from a specified location.") : FIOSTargetSettingsCustomizationConstants::DisabledTip);
 }
@@ -424,104 +644,6 @@ void FIOSTargetSettingsCustomization::CopySetupFilesIntoProject()
 	SavedLayoutBuilder->ForceRefreshDetails();
 }
 
-void FIOSTargetSettingsCustomization::OnPlistPropertyModified()
-{
-// 	check(SetupForPlatformAttribute.Get() == true);
-/*	const UIOSRuntimeSettings& Settings = *GetDefault<UIOSRuntimeSettings>();
-
-	FManifestUpdateHelper Updater(GameInfoPath);
-
-	// The text we're trying to replace looks like this:
-	// 	<key>UISupportedInterfaceOrientations</key>
-	// 	<array>
-	// 		<string>UIInterfaceOrientationLandscapeRight</string>
-	// 		<string>UIInterfaceOrientationLandscapeLeft</string>
-	// 	</array>
-	const FString InterfaceOrientations(TEXT("<key>UISupportedInterfaceOrientations</key>"));
-	const FString ClosingArray(TEXT("</array>"));
-
-	// Build the replacement array
-	FString OrientationArrayBody = TEXT("\n\t<array>\n");
-	if (Settings.bSupportsPortraitOrientation)
-	{
-		OrientationArrayBody += TEXT("\t\t<string>UIInterfaceOrientationPortrait</string>\n");
-	}
-	if (Settings.bSupportsUpsideDownOrientation)
-	{
-		OrientationArrayBody += TEXT("\t\t<string>UIInterfaceOrientationPortraitUpsideDown</string>\n");
-	}
-	if (Settings.bSupportsLandscapeLeftOrientation && (!Settings.bSupportsPortraitOrientation && !Settings.bSupportsUpsideDownOrientation))
-	{
-		OrientationArrayBody += TEXT("\t\t<string>UIInterfaceOrientationLandscapeLeft</string>\n");
-	}
-	if (Settings.bSupportsLandscapeRightOrientation && (!Settings.bSupportsPortraitOrientation && !Settings.bSupportsUpsideDownOrientation))
-	{
-		OrientationArrayBody += TEXT("\t\t<string>UIInterfaceOrientationLandscapeRight</string>\n");
-	}
-	OrientationArrayBody += TEXT("\t");
-	Updater.ReplaceKey(InterfaceOrientations, ClosingArray, OrientationArrayBody);
-
-	// build the replacement bundle display name
-	const FString BundleDisplayNameKey(TEXT("<key>CFBundleDisplayName</key>"));
-	const FString ClosingString(TEXT("</string>"));
-	FString BundleDisplayNameBody = TEXT("\n\t<string>") + Settings.BundleDisplayName;
-	Updater.ReplaceKey(BundleDisplayNameKey, ClosingString, BundleDisplayNameBody);
-
-	// build the replacement bundle display name
-	const FString BundleNameKey(TEXT("<key>CFBundleName</key>"));
-	FString BundleNameBody = TEXT("\n\t<string>") + Settings.BundleName;
-	Updater.ReplaceKey(BundleNameKey, ClosingString, BundleNameBody);
-
-	// build the replacement bundle identifier
-	const FString BundleIdentifierKey(TEXT("<key>CFBundleIdentifier</key>"));
-	FString BundleIdentifierBody = TEXT("\n\t<string>") + Settings.BundleIdentifier;
-	Updater.ReplaceKey(BundleIdentifierKey, ClosingString, BundleIdentifierBody);
-
-	// build the replacement version info
-	const FString BundleShortVersionKey(TEXT("<key>CFBundleShortVersionString</key>"));
-	FString VersionInfoBody = TEXT("\n\t<string>") + Settings.VersionInfo;
-	Updater.ReplaceKey(BundleShortVersionKey, ClosingString, VersionInfoBody);
-
-	// build the replacement required device caps
-	const FString RequiredDeviceCaps(TEXT("<key>UIRequiredDeviceCapabilities</key>"));
-	FString DeviceCapsArrayBody = TEXT("\n\t<array>\n");
-	// automatically add armv7 for now
-	DeviceCapsArrayBody += TEXT("\t\t<string>armv7</string>\n");
-	if (Settings.bSupportsOpenGLES2)
-	{
-		DeviceCapsArrayBody += TEXT("\t\t<string>opengles-2</string>\n");
-	}
-	else if (Settings.bSupportsMetal)
-	{
-		DeviceCapsArrayBody += TEXT("\t\t<string>metal</string>\n");
-	}
-	DeviceCapsArrayBody += TEXT("\t");
-	Updater.ReplaceKey(RequiredDeviceCaps, ClosingArray, DeviceCapsArrayBody);
-
-	// build the replacement min iOS version
-	const FString MiniOSVersionKey(TEXT("<key>MinimumOSVersion</key>"));
-	FString iOSVersionBody = TEXT("\n\t<string>");
-	switch (Settings.MinimumiOSVersion)
-	{
-	case EIOSVersion::IOS_61:
-		iOSVersionBody += TEXT("6.1");
-		break;
-	case EIOSVersion::IOS_7:
-		iOSVersionBody += TEXT("7.0");
-		break;
-	case EIOSVersion::IOS_8:
-		iOSVersionBody += TEXT("8.0");
-		break;
-	}
-	Updater.ReplaceKey(MiniOSVersionKey, ClosingString, iOSVersionBody);
-
-	const FString ClosingDict(TEXT("</dict>"));
-	Updater.ReplaceKey(iOSVersionBody + ClosingString, ClosingDict, Settings.AdditionalPlistData.Replace(TEXT("\\n"), TEXT("\n")));
-
-	// Write out the updated .plist
-	Updater.Finalize(GameInfoPath, true, FFileHelper::EEncodingOptions::ForceUTF8);*/
-}
-
 void FIOSTargetSettingsCustomization::BuildImageRow(IDetailLayoutBuilder& DetailLayout, IDetailCategoryBuilder& Category, const FPlatformIconInfo& Info, const FVector2D& MaxDisplaySize)
 {
 	const FString AutomaticImagePath = EngineGraphicsPath / Info.IconPath;
@@ -557,11 +679,23 @@ void FIOSTargetSettingsCustomization::BuildImageRow(IDetailLayoutBuilder& Detail
 		];
 }
 
-static FString OutputMessage;
-static void OnOutput(FString Message)
+void FIOSTargetSettingsCustomization::FindRequiredFiles()
 {
-	OutputMessage += Message;
-	UE_LOG(LogTemp, Display, TEXT("%s\n"), *Message);
+	const UIOSRuntimeSettings& Settings = *GetDefault<UIOSRuntimeSettings>();
+#if PLATFORM_MAC
+	FString CmdExe = TEXT("/bin/sh");
+	FString ScriptPath = FPaths::ConvertRelativePathToFull(FPaths::EngineDir() / TEXT("Build/BatchFiles/Mac/RunMono.sh"));
+	FString IPPPath = FPaths::ConvertRelativePathToFull(FPaths::EngineDir() / TEXT("Binaries/DotNet/IOS/IPhonePackager.exe"));
+	FString CommandLine = FString::Printf(TEXT("\"%s\" \"%s\" certificates Engine -bundlename \"%s\""), *ScriptPath, *IPPPath, *(Settings.BundleIdentifier));
+#else
+	FString CmdExe = FPaths::ConvertRelativePathToFull(FPaths::EngineDir() / TEXT("Binaries/DotNet/IOS/IPhonePackager.exe"));
+	FString CommandLine = FString::Printf(TEXT("certificates Engine -bundlename \"%s\""), *(Settings.BundleIdentifier));
+#endif
+	IPPProcess = MakeShareable(new FMonitoredProcess(CmdExe, CommandLine, true));
+	OutputMessage = TEXT("");
+	IPPProcess->OnOutput().BindStatic(&OnOutput);
+	IPPProcess->Launch();
+	TickerHandle = FTicker::GetCoreTicker().AddTicker(FTickerDelegate::CreateRaw(this, &FIOSTargetSettingsCustomization::UpdateStatusDelegate), 1.0f);
 }
 
 FReply FIOSTargetSettingsCustomization::OnInstallProvisionClicked()
@@ -624,14 +758,16 @@ FReply FIOSTargetSettingsCustomization::OnInstallProvisionClicked()
 				return FReply::Handled();
 			}
 		}
+
+		const UIOSRuntimeSettings& Settings = *GetDefault<UIOSRuntimeSettings>();
 #if PLATFORM_MAC
 		FString CmdExe = TEXT("/bin/sh");
 		FString ScriptPath = FPaths::ConvertRelativePathToFull(FPaths::EngineDir() / TEXT("Build/BatchFiles/Mac/RunMono.sh"));
 		FString IPPPath = FPaths::ConvertRelativePathToFull(FPaths::EngineDir() / TEXT("Binaries/DotNet/IOS/IPhonePackager.exe"));
-		FString CommandLine = FString::Printf(TEXT("\"%s\" \"%s\" Install Engine -project \"%s\" -provision \"%s\""), *ScriptPath, *IPPPath, *ProjectPath, *ProvisionPath);
+		FString CommandLine = FString::Printf(TEXT("\"%s\" \"%s\" Install Engine -project \"%s\" -provision \"%s\" -bundlename \"%s\""), *ScriptPath, *IPPPath, *ProjectPath, *ProvisionPath, *(Settings.BundleIdentifier));
 #else
 		FString CmdExe = FPaths::ConvertRelativePathToFull(FPaths::EngineDir() / TEXT("Binaries/DotNet/IOS/IPhonePackager.exe"));
-		FString CommandLine = FString::Printf(TEXT("Install Engine -project \"%s\" -provision \"%s\""), *ProjectPath, *ProvisionPath);
+		FString CommandLine = FString::Printf(TEXT("Install Engine -project \"%s\" -provision \"%s\" -bundlename \"%s\""), *ProjectPath, *ProvisionPath, *(Settings.BundleIdentifier));
 #endif
 		IPPProcess = MakeShareable(new FMonitoredProcess(CmdExe, CommandLine, true));
 		OutputMessage = TEXT("");
@@ -681,15 +817,16 @@ FReply FIOSTargetSettingsCustomization::OnInstallCertificateClicked()
 
 	if ( bOpened )
 	{
+		const UIOSRuntimeSettings& Settings = *GetDefault<UIOSRuntimeSettings>();
 		CertPath = FPaths::ConvertRelativePathToFull(OpenFilenames[0]);
 #if PLATFORM_MAC
 		FString CmdExe = TEXT("/bin/sh");
 		FString ScriptPath = FPaths::ConvertRelativePathToFull(FPaths::EngineDir() / TEXT("Build/BatchFiles/Mac/RunMono.sh"));
 		FString IPPPath = FPaths::ConvertRelativePathToFull(FPaths::EngineDir() / TEXT("Binaries/DotNet/IOS/IPhonePackager.exe"));
-		FString CommandLine = FString::Printf(TEXT("\"%s\" \"%s\" Install Engine -project \"%s\" -certificate \"%s\""), *ScriptPath, *IPPPath, *ProjectPath, *CertPath);
+		FString CommandLine = FString::Printf(TEXT("\"%s\" \"%s\" Install Engine -project \"%s\" -certificate \"%s\" -bundlename \"%s\""), *ScriptPath, *IPPPath, *ProjectPath, *CertPath, *(Settings.BundleIdentifier));
 #else
 		FString CmdExe = FPaths::ConvertRelativePathToFull(FPaths::EngineDir() / TEXT("Binaries/DotNet/IOS/IPhonePackager.exe"));
-		FString CommandLine = FString::Printf(TEXT("Install Engine -project \"%s\" -certificate \"%s\""), *ProjectPath, *CertPath);
+		FString CommandLine = FString::Printf(TEXT("Install Engine -project \"%s\" -certificate \"%s\" -bundlename \"%s\""), *ProjectPath, *CertPath, *(Settings.BundleIdentifier));
 #endif
 		IPPProcess = MakeShareable(new FMonitoredProcess(CmdExe, CommandLine, true));
 		OutputMessage = TEXT("");
@@ -743,6 +880,72 @@ bool FIOSTargetSettingsCustomization::UpdateStatusDelegate(float DeltaTime)
 
 	return false;
 }
+
+TSharedRef<ITableRow> FIOSTargetSettingsCustomization::HandleProvisionListGenerateRow( ProvisionPtr InProvision, const TSharedRef<STableViewBase>& OwnerTable )
+{
+	return SNew(SProvisionListRow, OwnerTable)
+		.Provision(InProvision);
+}
+
+TSharedRef<ITableRow> FIOSTargetSettingsCustomization::HandleCertificateListGenerateRow( CertificatePtr InCertificate, const TSharedRef<STableViewBase>& OwnerTable )
+{
+	return SNew(SCertificateListRow, OwnerTable)
+		.Certificate(InCertificate);
+}
+
+void FIOSTargetSettingsCustomization::HandleAllProvisionsHyperlinkNavigate( bool AllProvisions )
+{
+	bShowAllProvisions = AllProvisions;
+	FilterLists();
+}
+
+void FIOSTargetSettingsCustomization::HandleAllCertificatesHyperlinkNavigate( bool AllCertificates )
+{
+	bShowAllCertificates = AllCertificates;
+	FilterLists();
+}
+
+void FIOSTargetSettingsCustomization::FilterLists()
+{
+	FilteredProvisionList.Reset();
+	FilteredCertificateList.Reset();
+
+	for (int Index = 0; Index < ProvisionList.Num(); ++Index)
+	{
+		if (SelectedProvision.Contains(ProvisionList[Index]->Name) && SelectedFile.Contains(ProvisionList[Index]->FileName))
+		{
+			ProvisionList[Index]->bSelected = true;
+		}
+		else
+		{
+			ProvisionList[Index]->bSelected = false;
+		}
+		if (bShowAllProvisions || ProvisionList[Index]->Status.Contains("VALID"))
+		{
+			FilteredProvisionList.Add(ProvisionList[Index]);
+		}
+	}
+
+	for (int Index = 0; Index < CertificateList.Num(); ++Index)
+	{
+		if (SelectedCert.Contains(CertificateList[Index]->Name))
+		{
+			CertificateList[Index]->bSelected = true;
+		}
+		else
+		{
+			CertificateList[Index]->bSelected = false;
+		}
+		if (bShowAllCertificates || CertificateList[Index]->Status.Contains("VALID"))
+		{
+			FilteredCertificateList.Add(CertificateList[Index]);
+		}
+	}
+
+	CertificateListView->RequestListRefresh();
+	ProvisionListView->RequestListRefresh();
+}
+
 //////////////////////////////////////////////////////////////////////////
 
 #undef LOCTEXT_NAMESPACE
\ No newline at end of file
diff --git a/Engine/Source/Developer/iOS/IOSPlatformEditor/Private/IOSTargetSettingsCustomization.h b/Engine/Source/Developer/iOS/IOSPlatformEditor/Private/IOSTargetSettingsCustomization.h
--- a/Engine/Source/Developer/iOS/IOSPlatformEditor/Private/IOSTargetSettingsCustomization.h
+++ b/Engine/Source/Developer/iOS/IOSPlatformEditor/Private/IOSTargetSettingsCustomization.h
@@ -6,6 +6,31 @@
 #include "PropertyEditorModule.h"
 #include "IOSRuntimeSettings.h"
 
+//////////////////////////////////////////////////////////////////////////
+// FProvision structure
+class FProvision
+{
+public:
+	FString Name;
+	FString FileName;
+	FString Status;
+	bool bDistribution;
+	bool bSelected;
+};
+typedef TSharedPtr<class FProvision> ProvisionPtr;
+
+//////////////////////////////////////////////////////////////////////////
+// FCertificate structure
+class FCertificate
+{
+public:
+	FString Name;
+	FString Status;
+	FString Expires;
+	bool bSelected;
+};
+typedef TSharedPtr<class FCertificate> CertificatePtr;
+
 //////////////////////////////////////////////////////////////////////////
 // FIOSTargetSettingsCustomization
 
@@ -38,9 +63,21 @@ private:
 
 	bool bProvisionInstalled;
 	bool bCertificateInstalled;
+	bool bShowAllProvisions;
+	bool bShowAllCertificates;
 
 	TSharedPtr<FMonitoredProcess> IPPProcess;
 	FDelegateHandle TickerHandle;
+	TArray<ProvisionPtr> ProvisionList;
+	TArray<ProvisionPtr> FilteredProvisionList;
+	TSharedPtr<SListView<ProvisionPtr> > ProvisionListView;
+	TArray<CertificatePtr> CertificateList;
+	TArray<CertificatePtr> FilteredCertificateList;
+	TSharedPtr<SListView<CertificatePtr> > CertificateListView;
+
+	FString SelectedProvision;
+	FString SelectedFile;
+	FString SelectedCert;
 
 private:
 	FIOSTargetSettingsCustomization();
@@ -55,9 +92,6 @@ private:
 	// Copies the setup files for the platform into the project
 	void CopySetupFilesIntoProject();
 
-	// Called when a property that is really stored in the plist is modified
-	void OnPlistPropertyModified();
-
 	// Builds an image row
 	void BuildImageRow(class IDetailLayoutBuilder& DetailLayout, class IDetailCategoryBuilder& Category, const struct FPlatformIconInfo& Info, const FVector2D& MaxDisplaySize);
 
@@ -76,9 +110,27 @@ private:
 	// Get the image to display for the certificate status
 	const FSlateBrush* GetCertificateStatus() const;
 
+	// find the installed certificates and provisions
+	void FindRequiredFiles();
+
 	// Update the provision status
 	void UpdateStatus();
 
 	// status tick delay
 	bool UpdateStatusDelegate(float delay);
+
+	// handle provision list generate row
+	TSharedRef<ITableRow> HandleProvisionListGenerateRow( ProvisionPtr Provision, const TSharedRef<STableViewBase>& OwnerTable );
+
+	// handle certificate list generate row
+	TSharedRef<ITableRow> HandleCertificateListGenerateRow( CertificatePtr Certificate, const TSharedRef<STableViewBase>& OwnerTable );
+
+	// handle which set of provisions to view
+	void HandleAllProvisionsHyperlinkNavigate( bool AllProvisions );
+
+	// handle which set of provisions to view
+	void HandleAllCertificatesHyperlinkNavigate( bool AllCertificates );
+
+	// filter the lists based on the settings
+	void FilterLists();
 };
diff --git a/Engine/Source/Developer/iOS/IOSPlatformEditor/Private/SCertificateListRow.h b/Engine/Source/Developer/iOS/IOSPlatformEditor/Private/SCertificateListRow.h
new file mode 100644
--- /dev/null
+++ b/Engine/Source/Developer/iOS/IOSPlatformEditor/Private/SCertificateListRow.h
@@ -0,0 +1,133 @@
+// Copyright 1998-2015 Epic Games, Inc. All Rights Reserved.
+
+#pragma once
+
+/**
+ * Implements a row widget for the certificate list view.
+ */
+class SCertificateListRow
+	: public SMultiColumnTableRow<CertificatePtr>
+{
+public:
+
+	SLATE_BEGIN_ARGS(SCertificateListRow) { }
+		SLATE_ARGUMENT(CertificatePtr, Certificate)
+	SLATE_END_ARGS()
+
+public:
+
+	/**
+	 * Constructs the widget.
+	 *
+	 * @param InArgs The arguments.
+	 */
+	void Construct( const FArguments& InArgs, const TSharedRef<STableViewBase>& InOwnerTableView )
+	{
+		Certificate = InArgs._Certificate;
+		
+		SMultiColumnTableRow<CertificatePtr>::Construct(FSuperRowType::FArguments(), InOwnerTableView);
+	}
+
+public:
+
+	/**
+	 * Generates the widget for the specified column.
+	 *
+	 * @param ColumnName The name of the column to generate the widget for.
+	 * @return The widget.
+	 */
+	BEGIN_SLATE_FUNCTION_BUILD_OPTIMIZATION
+	virtual TSharedRef<SWidget> GenerateWidgetForColumn( const FName& ColumnName ) override
+	{
+		if (ColumnName == TEXT("Name"))
+		{
+			return SNew(SBox)
+				.Padding(FMargin(4.0f, 0.0f))
+				.VAlign(VAlign_Center)
+				[
+					SNew(STextBlock)
+						.ColorAndOpacity(this, &SCertificateListRow::HandleSelectedColorAndOpacity)
+						.Text(this, &SCertificateListRow::HandleNameText)
+				];
+		}
+		else if (ColumnName == TEXT("Status"))
+		{
+			return SNew(SBox)
+				.Padding(FMargin(4.0f, 0.0f))
+				.VAlign(VAlign_Center)
+				[
+					SNew(STextBlock)
+						.ColorAndOpacity(this, &SCertificateListRow::HandleStatusTextColorAndOpacity)
+						.Text(this, &SCertificateListRow::HandleStatusTextBlockText)
+				];
+		}
+		else if (ColumnName == TEXT("Expires"))
+		{
+			return SNew(SBox)
+				.Padding(FMargin(4.0f, 0.0f))
+				.VAlign(VAlign_Center)
+				.HAlign(HAlign_Center)
+				[
+					SNew(STextBlock)
+					.ColorAndOpacity(this, &SCertificateListRow::HandleSelectedColorAndOpacity)
+					.Text(this, &SCertificateListRow::HandleExpiresText)
+				];
+		}
+
+		return SNullWidget::NullWidget;
+	}
+	END_SLATE_FUNCTION_BUILD_OPTIMIZATION
+
+private:
+
+	// Callback for getting the text in the 'Name' column.
+	FText HandleNameText( ) const
+	{
+		return FText::FromString(Certificate->Name);
+	}
+
+	// Callback for getting the status text.
+	FText HandleStatusTextBlockText( ) const
+	{
+		if (Certificate->Status == TEXT("EXPIRED"))
+		{
+			return FText::FromString(TEXT("Expired"));
+		}
+		return FText::FromString(TEXT("Valid"));
+	}
+
+	// Callback for getting the status text color.
+	FSlateColor HandleStatusTextColorAndOpacity( ) const
+	{
+		if (Certificate->Status == TEXT("EXPIRED"))
+		{
+			return FSlateColor(FLinearColor(1.0f, 0.0f, 0.0f));
+		}
+		else if (Certificate->bSelected)
+		{
+			return FSlateColor(FLinearColor(0.0f, 1.0f, 0.0f));
+		}
+		return FSlateColor(FLinearColor(1.0f, 1.0f, 1.0f));
+	}
+
+	// Callback for getting the status text color.
+	FSlateColor HandleSelectedColorAndOpacity( ) const
+	{
+		if (Certificate->bSelected)
+		{
+			return FSlateColor(FLinearColor(0.0f, 1.0f, 0.0f));
+		}
+		return FSlateColor(FLinearColor(1.0f, 1.0f, 1.0f));
+	}
+
+	// Callback for getting the text in the 'Expires' column.
+	FText HandleExpiresText( ) const
+	{
+		return FText::FromString(Certificate->Expires);
+	}
+
+private:
+
+	// Holds the target device service used to populate this row.
+	CertificatePtr Certificate;
+};
diff --git a/Engine/Source/Developer/iOS/IOSPlatformEditor/Private/SProvisionListRow.h b/Engine/Source/Developer/iOS/IOSPlatformEditor/Private/SProvisionListRow.h
new file mode 100644
--- /dev/null
+++ b/Engine/Source/Developer/iOS/IOSPlatformEditor/Private/SProvisionListRow.h
@@ -0,0 +1,185 @@
+// Copyright 1998-2015 Epic Games, Inc. All Rights Reserved.
+
+#pragma once
+
+#define IMAGE_BRUSH( RelativePath, ... ) FSlateImageBrush( FPaths::EngineContentDir() / TEXT("Editor/Slate") / RelativePath + TEXT(".png"), __VA_ARGS__ )
+
+/**
+ * Implements a row widget for the provision list view.
+ */
+class SProvisionListRow
+	: public SMultiColumnTableRow<ProvisionPtr>
+{
+public:
+
+	SLATE_BEGIN_ARGS(SProvisionListRow) { }
+		SLATE_ARGUMENT(ProvisionPtr, Provision)
+	SLATE_END_ARGS()
+
+public:
+
+	/**
+	 * Constructs the widget.
+	 *
+	 * @param InArgs The arguments.
+	 */
+	void Construct( const FArguments& InArgs, const TSharedRef<STableViewBase>& InOwnerTableView )
+	{
+		Provision = InArgs._Provision;
+		
+		SMultiColumnTableRow<ProvisionPtr>::Construct(FSuperRowType::FArguments(), InOwnerTableView);
+
+		/* Set images for various SCheckBox states ... */
+		if (!bInitialized)
+		{
+			ProvisionCheckBoxStyle = FCheckBoxStyle()
+				.SetCheckBoxType(ESlateCheckBoxType::CheckBox)
+				.SetUncheckedImage( FSlateNoResource() )
+				.SetUncheckedHoveredImage( FSlateNoResource() )
+				.SetUncheckedPressedImage( FSlateNoResource() )
+				.SetCheckedImage( IMAGE_BRUSH( "Automation/Success", FVector2D(16.0f, 16.0f) ) )
+				.SetCheckedHoveredImage( IMAGE_BRUSH( "Automation/Success", FVector2D(16.0f, 16.0f), FLinearColor( 0.5f, 0.5f, 0.5f ) ) )
+				.SetCheckedPressedImage( IMAGE_BRUSH( "Automation/Success", FVector2D(16.0f, 16.0f) ) )
+				.SetUndeterminedImage(FSlateNoResource() )
+				.SetUndeterminedHoveredImage( FSlateNoResource() )
+				.SetUndeterminedPressedImage( FSlateNoResource() );
+
+			bInitialized = true;
+		}
+	}
+
+public:
+
+	/**
+	 * Generates the widget for the specified column.
+	 *
+	 * @param ColumnName The name of the column to generate the widget for.
+	 * @return The widget.
+	 */
+	BEGIN_SLATE_FUNCTION_BUILD_OPTIMIZATION
+	virtual TSharedRef<SWidget> GenerateWidgetForColumn( const FName& ColumnName ) override
+	{
+		if (ColumnName == TEXT("Name"))
+		{
+			return SNew(SBox)
+				.Padding(FMargin(4.0f, 0.0f))
+				.VAlign(VAlign_Center)
+				[
+					SNew(STextBlock)
+						.ColorAndOpacity(this, &SProvisionListRow::HandleSelectedColorAndOpacity)
+						.Text(this, &SProvisionListRow::HandleNameText)
+				];
+		}
+		else if (ColumnName == TEXT("File"))
+		{
+			return SNew(SBox)
+				.Padding(FMargin(4.0f, 0.0f))
+				.VAlign(VAlign_Center)
+				[
+					SNew(STextBlock)
+					.ColorAndOpacity(this, &SProvisionListRow::HandleSelectedColorAndOpacity)
+					.Text(this, &SProvisionListRow::HandleFileText)
+				];
+		}
+		else if (ColumnName == TEXT("Status"))
+		{
+			return SNew(SBox)
+				.Padding(FMargin(4.0f, 0.0f))
+				.VAlign(VAlign_Center)
+				[
+					SNew(STextBlock)
+						.ColorAndOpacity(this, &SProvisionListRow::HandleStatusTextColorAndOpacity)
+						.Text(this, &SProvisionListRow::HandleStatusTextBlockText)
+				];
+		}
+		else if (ColumnName == TEXT("Distribution"))
+		{
+			return SNew(SBox)
+				.Padding(FMargin(4.0f, 0.0f))
+				.VAlign(VAlign_Center)
+				.HAlign(HAlign_Center)
+				[
+					SNew(SCheckBox)
+						.IsChecked(this, &SProvisionListRow::HandleDistribution)
+						.Style(&ProvisionCheckBoxStyle)
+				];
+		}
+
+		return SNullWidget::NullWidget;
+	}
+	END_SLATE_FUNCTION_BUILD_OPTIMIZATION
+
+private:
+
+	// Callback for getting the text in the 'Name' column.
+	FText HandleNameText( ) const
+	{
+		return FText::FromString(Provision->Name);
+	}
+
+	// Callback for getting the text in the 'File' column.
+	FText HandleFileText( ) const
+	{
+		return FText::FromString(Provision->FileName);
+	}
+
+	// Callback for getting the status text.
+	FText HandleStatusTextBlockText( ) const
+	{
+		if (Provision->Status == TEXT("NO_MATCH"))
+		{
+			return FText::FromString(TEXT("Identifier Not Matched"));
+		}
+		else if (Provision->Status == TEXT("NO_CERT"))
+		{
+			return FText::FromString(TEXT("No Valid Certificate Found"));
+		}
+		else if (Provision->Status == TEXT("EXPIRED"))
+		{
+			return FText::FromString(TEXT("Expired"));
+		}
+		return FText::FromString(TEXT("Valid"));
+	}
+
+	// Callback for getting the status text color.
+	FSlateColor HandleStatusTextColorAndOpacity( ) const
+	{
+		if (Provision->Status == TEXT("NO_MATCH"))
+		{
+			return FSlateColor(FLinearColor(1.0f, 1.0f, 0.0f));
+		}
+		else if (Provision->Status == TEXT("NO_CERT") || Provision->Status == TEXT("EXPIRED"))
+		{
+			return FSlateColor(FLinearColor(1.0f, 0.0f, 0.0f));
+		}
+		else if (Provision->bSelected)
+		{
+			return FSlateColor(FLinearColor(0.0f, 1.0f, 0.0f));
+		}
+		return FSlateColor(FLinearColor(1.0f, 1.0f, 1.0f));
+	}
+
+	// Callback for getting the status text color.
+	FSlateColor HandleSelectedColorAndOpacity( ) const
+	{
+		if (Provision->bSelected)
+		{
+			return FSlateColor(FLinearColor(0.0f, 1.0f, 0.0f));
+		}
+		return FSlateColor(FLinearColor(1.0f, 1.0f, 1.0f));
+	}
+
+	// Callback to determine distribution
+	ECheckBoxState HandleDistribution( ) const
+	{
+		return Provision->bDistribution ? ECheckBoxState::Checked : ECheckBoxState::Unchecked;
+	}
+
+private:
+
+	// Holds the target device service used to populate this row.
+	ProvisionPtr Provision;
+
+	static bool bInitialized;
+	static FCheckBoxStyle ProvisionCheckBoxStyle;
+};
diff --git a/Engine/Source/Programs/IOS/iPhonePackager/CodeSigning.cs b/Engine/Source/Programs/IOS/iPhonePackager/CodeSigning.cs
--- a/Engine/Source/Programs/IOS/iPhonePackager/CodeSigning.cs
+++ b/Engine/Source/Programs/IOS/iPhonePackager/CodeSigning.cs
@@ -211,6 +211,146 @@ namespace iPhonePackager
 			}
 		}
 
+		/// <summary>
+		/// Finds all valid installed certificates
+		/// </summary>
+		public static void FindCertificates()
+		{
+			if (Environment.OSVersion.Platform == PlatformID.Unix || Environment.OSVersion.Platform == PlatformID.MacOSX)
+			{
+				// run certtool y to get the currently installed certificates
+				CertToolData = "";
+				Process CertTool = new Process();
+				CertTool.StartInfo.FileName = "/usr/bin/security";
+				CertTool.StartInfo.UseShellExecute = false;
+				CertTool.StartInfo.Arguments = "find-certificate -a -c \"iPhone\" -p";
+				CertTool.StartInfo.RedirectStandardOutput = true;
+				CertTool.OutputDataReceived += new DataReceivedEventHandler(OutputReceivedCertToolProcessCall);
+				CertTool.Start();
+				CertTool.BeginOutputReadLine();
+				CertTool.WaitForExit();
+				if (CertTool.ExitCode == 0)
+				{
+					string header = "-----BEGIN CERTIFICATE-----\n";
+					string footer = "-----END CERTIFICATE-----";
+					int start = CertToolData.IndexOf(header);
+					while (start != -1)
+					{
+						start += header.Length;
+						int end = CertToolData.IndexOf(footer, start);
+						string base64 = CertToolData.Substring(start, (end - start));
+						byte[] certData = Convert.FromBase64String(base64);
+						X509Certificate2 cert = new X509Certificate2(certData);
+						//@TODO: Pretty sure the certificate information from the library is in local time, not UTC and this works as expected, but it should be verified!
+						DateTime EffectiveDate = cert.NotBefore;
+						DateTime ExpirationDate = cert.NotAfter;
+						DateTime Now = DateTime.Now;
+						string Subject = cert.Subject;
+						int SubjStart = Subject.IndexOf("CN=") + 3;
+						int SubjEnd = Subject.IndexOf(",", SubjStart);
+						cert.FriendlyName = Subject.Substring(SubjStart, (SubjEnd - SubjStart));
+						bool bCertTimeIsValid = (EffectiveDate < Now) && (ExpirationDate > Now);
+						Program.LogVerbose("CERTIFICATE-Name:{0},Validity:{1},StartDate:{2},EndDate:{3}", cert.FriendlyName, bCertTimeIsValid ? "VALID" : "EXPIRED", cert.GetEffectiveDateString(), cert.GetExpirationDateString());
+
+						start = CertToolData.IndexOf(header, start);
+					}
+				}
+			}
+			else
+			{
+				// Open the personal certificate store on this machine
+				X509Store Store = new X509Store();
+				Store.Open(OpenFlags.ReadOnly);
+
+				// Try finding all certificates that match either iPhone Developer or iPhone Distribution
+				X509Certificate2Collection FoundCerts = Store.Certificates.Find(X509FindType.FindBySubjectName, "iPhone Developer", false);
+
+				foreach (X509Certificate2 TestCert in FoundCerts)
+				{
+					//@TODO: Pretty sure the certificate information from the library is in local time, not UTC and this works as expected, but it should be verified!
+					DateTime EffectiveDate = TestCert.NotBefore;
+					DateTime ExpirationDate = TestCert.NotAfter;
+					DateTime Now = DateTime.Now;
+
+					bool bCertTimeIsValid = (EffectiveDate < Now) && (ExpirationDate > Now);
+					Program.LogVerbose("CERTIFICATE-Name:{0},Validity:{1},StartDate:{2},EndDate:{3}", TestCert.FriendlyName, bCertTimeIsValid ? "VALID" : "EXPIRED", TestCert.GetEffectiveDateString(), TestCert.GetExpirationDateString());
+				}
+
+				FoundCerts = Store.Certificates.Find(X509FindType.FindBySubjectName, "iPhone Distribution", false);
+
+				foreach (X509Certificate2 TestCert in FoundCerts)
+				{
+					//@TODO: Pretty sure the certificate information from the library is in local time, not UTC and this works as expected, but it should be verified!
+					DateTime EffectiveDate = TestCert.NotBefore;
+					DateTime ExpirationDate = TestCert.NotAfter;
+					DateTime Now = DateTime.Now;
+
+					bool bCertTimeIsValid = (EffectiveDate < Now) && (ExpirationDate > Now);
+					Program.LogVerbose("CERTIFICATE-Name:{0},Validity:{1},StartDate:{2},EndDate:{3}", TestCert.FriendlyName, bCertTimeIsValid ? "VALID" : "EXPIRED", TestCert.GetEffectiveDateString(), TestCert.GetExpirationDateString());
+				}
+
+				Store.Close();
+			}
+		}
+
+		/// <summary>
+		/// Finds all valid installed provisions
+		/// </summary>
+		public static void FindProvisions(string CFBundleIdentifier)
+		{
+			// cache the provision library
+			string SelectedProvision = "";
+			string SelectedCert = "";
+			string SelectedFile = "";
+			bool bFoundName = false;
+			Dictionary<string, MobileProvision> ProvisionLibrary = new Dictionary<string, MobileProvision>();
+			foreach (string Provision in Directory.EnumerateFiles(Config.ProvisionDirectory, "*.mobileprovision"))
+			{
+				MobileProvision p = MobileProvisionParser.ParseFile(Provision);
+
+				DateTime EffectiveDate = p.CreationDate;
+				DateTime ExpirationDate = p.ExpirationDate;
+				DateTime Now = DateTime.Now;
+
+				bool bCertTimeIsValid = (EffectiveDate < Now) && (ExpirationDate > Now);
+				bool bValid = false;
+				X509Certificate2 Cert = FindCertificate(p);
+				if (Cert != null)
+				{
+					bValid = (Cert.NotBefore < Now) && (Cert.NotAfter > Now);
+				}
+				bool bPassesNameCheck = p.ApplicationIdentifier.Substring(p.ApplicationIdentifierPrefix.Length+1) == CFBundleIdentifier;
+				bool bPassesWildCardCheck = p.ApplicationIdentifier.Contains("*");
+				bool bDistribution = ((p.ProvisionedDeviceIDs.Count == 0) && !p.bDebug);
+				string Validity = "VALID";
+				if (!bPassesNameCheck && !bPassesWildCardCheck)
+				{
+					Validity = "NO_MATCH";
+				}
+				else
+				{
+					if (!bValid)
+					{
+						Validity = "NO_CERT";
+					}
+					else if (!bCertTimeIsValid)
+					{
+						Validity = "EXPIRED";
+					}
+				}
+				if ((string.IsNullOrWhiteSpace(SelectedProvision) || !bFoundName) && Validity == "VALID" && !bDistribution)
+				{
+					bFoundName = bPassesNameCheck;
+					SelectedProvision = p.ProvisionName;
+					SelectedFile = Path.GetFileName(Provision);
+					SelectedCert = Cert.FriendlyName;
+				}
+				Program.LogVerbose("PROVISION-File:{0},Name:{1},Validity:{2},StartDate:{3},EndDate:{4},Type:{5}", Path.GetFileName(Provision), p.ProvisionName, Validity, EffectiveDate.ToString(), ExpirationDate.ToString(), bDistribution ? "DISTRIBUTION" : "DEVELOPMENT");
+			}
+
+			Program.LogVerbose("MATCHED-Provision:{0},File:{1},Cert:{2}", SelectedProvision, SelectedFile, SelectedCert);
+		}
+
 		/// <summary>
 		/// Tries to find a matching certificate on this machine from the the serial number of one of the
 		/// certificates in the mobile provision (the one in the mobileprovision is missing the public/private key pair)
diff --git a/Engine/Source/Programs/IOS/iPhonePackager/iPhonePackager.cs b/Engine/Source/Programs/IOS/iPhonePackager/iPhonePackager.cs
--- a/Engine/Source/Programs/IOS/iPhonePackager/iPhonePackager.cs
+++ b/Engine/Source/Programs/IOS/iPhonePackager/iPhonePackager.cs
@@ -773,6 +773,15 @@ namespace iPhonePackager
 						{
 							ToolsHub.TryInstallingCertificate_PromptForKey(Config.Certificate, false);
 						}
+						CodeSignatureBuilder.FindCertificates();
+						CodeSignatureBuilder.FindProvisions(Config.OverrideBundleName);
+						break;
+
+					case "certificates":
+						{
+							CodeSignatureBuilder.FindCertificates();
+							CodeSignatureBuilder.FindProvisions(Config.OverrideBundleName);
+						}
 						break;
 
 					case "resigntool":
