addition of more information at various stages of the provision/certificate detection process
#ios

diff --git a/Engine/Source/Developer/iOS/IOSPlatformEditor/Private/IOSTargetSettingsCustomization.cpp b/Engine/Source/Developer/iOS/IOSPlatformEditor/Private/IOSTargetSettingsCustomization.cpp
--- a/Engine/Source/Developer/iOS/IOSPlatformEditor/Private/IOSTargetSettingsCustomization.cpp
+++ b/Engine/Source/Developer/iOS/IOSPlatformEditor/Private/IOSTargetSettingsCustomization.cpp
@@ -1,6 +1,7 @@
 // Copyright 1998-2015 Epic Games, Inc. All Rights Reserved.
 
 #include "IOSPlatformEditorPrivatePCH.h"
+#include "SWidgetSwitcher.h"
 #include "IOSTargetSettingsCustomization.h"
 #include "DetailLayoutBuilder.h"
 #include "DetailCategoryBuilder.h"
@@ -90,13 +91,13 @@ void FIOSTargetSettingsCustomization::CustomizeDetails(IDetailLayoutBuilder& Det
 {
 	SavedLayoutBuilder = &DetailLayout;
 
-	FindRequiredFiles();
-
 	BuildPListSection(DetailLayout);
 
 	BuildIconSection(DetailLayout);
 
 	BuildRemoteBuildingSection(DetailLayout);
+
+	FindRequiredFiles();
 }
 
 static FString OutputMessage;
@@ -254,79 +255,117 @@ void FIOSTargetSettingsCustomization::BuildPListSection(IDetailLayoutBuilder& De
 		[
 			SNew(SVerticalBox)
 			+ SVerticalBox::Slot()
-			.Padding(FMargin(10, 10, 10, 10))
-			.AutoHeight()
-			[
-				SAssignNew(ProvisionListView, SListView<ProvisionPtr>)
-				.ItemHeight(20.0f)
-				.ListItemsSource(&FilteredProvisionList)
-				.OnGenerateRow(this, &FIOSTargetSettingsCustomization::HandleProvisionListGenerateRow)
-				.SelectionMode(ESelectionMode::None)
-				.HeaderRow
-				(
-				SNew(SHeaderRow)
-				+ SHeaderRow::Column("Name")
-				.DefaultLabel(LOCTEXT("ProvisionListNameColumnHeader", "Provision"))
-				.FillWidth(1.0f)
-				+ SHeaderRow::Column("File")
-				.DefaultLabel(LOCTEXT("ProvisionListFileColumnHeader", "File"))
-				+ SHeaderRow::Column("Status")
-				.DefaultLabel(LOCTEXT("ProvisionListStatusColumnHeader", "Status"))
-				+ SHeaderRow::Column("Distribution")
-				.DefaultLabel(LOCTEXT("ProvisionListDistributionColumnHeader", "Distribution"))
-				.FixedWidth(75.0f)
-				)
-			]
-			+ SVerticalBox::Slot()
-			.AutoHeight()
-			.Padding(0.0f, 6.0f, 0.0f, 4.0f)
-			[
-				SNew(SSeparator)
-				.Orientation(Orient_Horizontal)
-			]
-			+ SVerticalBox::Slot()
 			.AutoHeight()
 			[
-				SNew(SHorizontalBox)
-
-				+ SHorizontalBox::Slot()
-				.AutoWidth()
+				SAssignNew(ProvisionInfoSwitcher, SWidgetSwitcher)
+				.WidgetIndex(0)
+				// searching for provisions
+				+SWidgetSwitcher::Slot()
+				.HAlign(HAlign_Fill)
+				.VAlign(VAlign_Center)
 				[
-					SNew(SRichTextBlock)
-					.Text(LOCTEXT("ProvisionMessage", "<RichTextBlock.TextHighlight>Note</>: The provision in green will be used to provision the IPA."))
-					.TextStyle(FEditorStyle::Get(), "MessageLog")
-					.DecoratorStyleSet(&FEditorStyle::Get())
-					.AutoWrapText(true)
+					SNew(SBorder)
+					.Padding(4)
+					[
+						SNew( STextBlock )
+						.Text( LOCTEXT( "ProvisionViewerFindingProvisions", "Please wait while we gather information." ) )
+						.AutoWrapText( true )
+					]
 				]
-
-				+ SHorizontalBox::Slot()
-				.FillWidth(1.0f)
-				.HAlign(HAlign_Right)
+				// no provisions found or no valid provisions
+				+SWidgetSwitcher::Slot()
+				.HAlign(HAlign_Fill)
+				.VAlign(VAlign_Center)
 				[
-					SNew(STextBlock)
-					.Text(LOCTEXT("ViewLabel", "View:"))
-				]
-
-				+ SHorizontalBox::Slot()
-					.AutoWidth()
-					.Padding(8.0f, 0.0f)
+					SNew(SBorder)
+					.Padding(4)
 					[
-						// all provisions hyper link
-						SNew(SHyperlink)
-						.OnNavigate(this, &FIOSTargetSettingsCustomization::HandleAllProvisionsHyperlinkNavigate, true)
-						.Text(LOCTEXT("AllProvisionsHyperLinkLabel", "All"))
-						.ToolTipText(LOCTEXT("AllProvisionsButtonTooltip", "View all provisions."))
+						SNew( STextBlock )
+						.Text( LOCTEXT( "ProvisionViewerNoValidProvisions", "No Provisions Found. Please Import a Provision." ) )
+						.AutoWrapText( true )
 					]
-
-				+ SHorizontalBox::Slot()
-					.AutoWidth()
+				]
+				+SWidgetSwitcher::Slot()
+				.HAlign(HAlign_Fill)
+				.VAlign(VAlign_Center)
+				[
+					SNew(SVerticalBox)
+					+ SVerticalBox::Slot()
+					.Padding(FMargin(10, 10, 10, 10))
+					.AutoHeight()
 					[
-						// valid provisions hyper link
-						SNew(SHyperlink)
-						.OnNavigate(this, &FIOSTargetSettingsCustomization::HandleAllProvisionsHyperlinkNavigate, false)
-						.Text(LOCTEXT("ValidProvisionsHyperlinkLabel", "Valid Only"))
-						.ToolTipText(LOCTEXT("ValidProvisionsHyperlinkTooltip", "View Valid provisions."))
+						SAssignNew(ProvisionListView, SListView<ProvisionPtr>)
+						.ItemHeight(20.0f)
+						.ListItemsSource(&FilteredProvisionList)
+						.OnGenerateRow(this, &FIOSTargetSettingsCustomization::HandleProvisionListGenerateRow)
+						.SelectionMode(ESelectionMode::None)
+						.HeaderRow
+						(
+						SNew(SHeaderRow)
+						+ SHeaderRow::Column("Name")
+						.DefaultLabel(LOCTEXT("ProvisionListNameColumnHeader", "Provision"))
+						.FillWidth(1.0f)
+						+ SHeaderRow::Column("File")
+						.DefaultLabel(LOCTEXT("ProvisionListFileColumnHeader", "File"))
+						+ SHeaderRow::Column("Status")
+						.DefaultLabel(LOCTEXT("ProvisionListStatusColumnHeader", "Status"))
+						+ SHeaderRow::Column("Distribution")
+						.DefaultLabel(LOCTEXT("ProvisionListDistributionColumnHeader", "Distribution"))
+						.FixedWidth(75.0f)
+						)
 					]
+					+ SVerticalBox::Slot()
+						.AutoHeight()
+						.Padding(0.0f, 6.0f, 0.0f, 4.0f)
+						[
+							SNew(SSeparator)
+							.Orientation(Orient_Horizontal)
+						]
+					+ SVerticalBox::Slot()
+						.AutoHeight()
+						[
+							SNew(SHorizontalBox)
+
+							+ SHorizontalBox::Slot()
+							.AutoWidth()
+							[
+								SNew(SRichTextBlock)
+								.Text(LOCTEXT("ProvisionMessage", "<RichTextBlock.TextHighlight>Note</>: The provision in green will be used to provision the IPA."))
+								.TextStyle(FEditorStyle::Get(), "MessageLog")
+								.DecoratorStyleSet(&FEditorStyle::Get())
+								.AutoWrapText(true)
+							]
+
+							+ SHorizontalBox::Slot()
+								.FillWidth(1.0f)
+								.HAlign(HAlign_Right)
+								[
+									SNew(STextBlock)
+									.Text(LOCTEXT("ViewLabel", "View:"))
+								]
+
+							+ SHorizontalBox::Slot()
+								.AutoWidth()
+								.Padding(8.0f, 0.0f)
+								[
+									// all provisions hyper link
+									SNew(SHyperlink)
+									.OnNavigate(this, &FIOSTargetSettingsCustomization::HandleAllProvisionsHyperlinkNavigate, true)
+									.Text(LOCTEXT("AllProvisionsHyperLinkLabel", "All"))
+									.ToolTipText(LOCTEXT("AllProvisionsButtonTooltip", "View all provisions."))
+								]
+
+							+ SHorizontalBox::Slot()
+								.AutoWidth()
+								[
+									// valid provisions hyper link
+									SNew(SHyperlink)
+									.OnNavigate(this, &FIOSTargetSettingsCustomization::HandleAllProvisionsHyperlinkNavigate, false)
+									.Text(LOCTEXT("ValidProvisionsHyperlinkLabel", "Valid Only"))
+									.ToolTipText(LOCTEXT("ValidProvisionsHyperlinkTooltip", "View Valid provisions."))
+								]
+						]
+				]
 			]
 			+ SVerticalBox::Slot()
 			.AutoHeight()
@@ -358,78 +397,116 @@ void FIOSTargetSettingsCustomization::BuildPListSection(IDetailLayoutBuilder& De
 			+ SVerticalBox::Slot()
 			.AutoHeight()
 			[
-				SNew(SHorizontalBox)
-				+ SHorizontalBox::Slot()
-				.Padding(FMargin(10, 10, 10, 10))
-				.FillWidth(1.0f)
+				SAssignNew(CertificateInfoSwitcher, SWidgetSwitcher)
+				.WidgetIndex(0)
+				// searching for provisions
+				+SWidgetSwitcher::Slot()
+				.HAlign(HAlign_Fill)
+				.VAlign(VAlign_Center)
 				[
-					SAssignNew(CertificateListView, SListView<CertificatePtr>)
-					.ItemHeight(20.0f)
-					.ListItemsSource(&FilteredCertificateList)
-					.OnGenerateRow(this, &FIOSTargetSettingsCustomization::HandleCertificateListGenerateRow)
-					.SelectionMode(ESelectionMode::None)
-					.HeaderRow
-					(
-						SNew(SHeaderRow)
-						+ SHeaderRow::Column("Name")
-						.DefaultLabel(LOCTEXT("CertificateListNameColumnHeader", "Certificate"))
-						+ SHeaderRow::Column("Status")
-						.DefaultLabel(LOCTEXT("CertificateListStatusColumnHeader", "Status"))
-						.FixedWidth(75.0f)
-						+ SHeaderRow::Column("Expires")
-						.DefaultLabel(LOCTEXT("CertificateListExpiresColumnHeader", "Expires"))
-						.FixedWidth(75.0f)
-					)
-				]
-			]
-			+ SVerticalBox::Slot()
-				.AutoHeight()
-				.Padding(0.0f, 6.0f, 0.0f, 4.0f)
-				[
-					SNew(SSeparator)
-					.Orientation(Orient_Horizontal)
-				]
-			+ SVerticalBox::Slot()
-				.AutoHeight()
-				[
-					SNew(SHorizontalBox)
-					+ SHorizontalBox::Slot()
-					.AutoWidth()
+					SNew(SBorder)
+					.Padding(4)
 					[
-						SNew(SRichTextBlock)
-						.Text(LOCTEXT("CertificateMessage", "<RichTextBlock.TextHighlight>Note</>: The certificate in green will be used to sign the IPA."))
-						.TextStyle(FEditorStyle::Get(), "MessageLog")
-						.DecoratorStyleSet(&FEditorStyle::Get())
-						.AutoWrapText(true)
+						SNew( STextBlock )
+						.Text( LOCTEXT( "CertificateViewerFindingProvisions", "Please wait while we gather information." ) )
+						.AutoWrapText( true )
 					]
-					+ SHorizontalBox::Slot()
-					.FillWidth(1.0f)
-					.HAlign(HAlign_Right)
+				]
+				// no provisions found or no valid provisions
+				+SWidgetSwitcher::Slot()
+					.HAlign(HAlign_Fill)
+					.VAlign(VAlign_Center)
 					[
-						SNew(STextBlock)
-						.Text(LOCTEXT("ViewLabel", "View:"))
-					]
-
-					+ SHorizontalBox::Slot()
-						.AutoWidth()
-						.Padding(8.0f, 0.0f)
+						SNew(SBorder)
+						.Padding(4)
 						[
-							// all provisions hyper link
-							SNew(SHyperlink)
-							.OnNavigate(this, &FIOSTargetSettingsCustomization::HandleAllCertificatesHyperlinkNavigate, true)
-							.Text(LOCTEXT("AllCertificatesHyperLinkLabel", "All"))
-							.ToolTipText(LOCTEXT("AllCertificatesButtonTooltip", "View all certificates."))
+							SNew( STextBlock )
+							.Text( LOCTEXT( "CertificateViewerNoValidProvisions", "No Certificates Found.  Please Import a Certificate." ) )
+							.AutoWrapText( true )
 						]
-
-					+ SHorizontalBox::Slot()
-						.AutoWidth()
+					]
+				+SWidgetSwitcher::Slot()
+					.HAlign(HAlign_Fill)
+					.VAlign(VAlign_Center)
+					[
+						SNew(SVerticalBox)
+						+ SVerticalBox::Slot()
+						.AutoHeight()
 						[
-							// valid provisions hyper link
-							SNew(SHyperlink)
-							.OnNavigate(this, &FIOSTargetSettingsCustomization::HandleAllCertificatesHyperlinkNavigate, false)
-							.Text(LOCTEXT("ValidCertificatesHyperlinkLabel", "Valid Only"))
-							.ToolTipText(LOCTEXT("ValidCertificatesHyperlinkTooltip", "View Valid certificates."))
+							SNew(SHorizontalBox)
+							+ SHorizontalBox::Slot()
+							.Padding(FMargin(10, 10, 10, 10))
+							.FillWidth(1.0f)
+							[
+								SAssignNew(CertificateListView, SListView<CertificatePtr>)
+								.ItemHeight(20.0f)
+								.ListItemsSource(&FilteredCertificateList)
+								.OnGenerateRow(this, &FIOSTargetSettingsCustomization::HandleCertificateListGenerateRow)
+								.SelectionMode(ESelectionMode::None)
+								.HeaderRow
+								(
+									SNew(SHeaderRow)
+									+ SHeaderRow::Column("Name")
+									.DefaultLabel(LOCTEXT("CertificateListNameColumnHeader", "Certificate"))
+									+ SHeaderRow::Column("Status")
+									.DefaultLabel(LOCTEXT("CertificateListStatusColumnHeader", "Status"))
+									.FixedWidth(75.0f)
+									+ SHeaderRow::Column("Expires")
+									.DefaultLabel(LOCTEXT("CertificateListExpiresColumnHeader", "Expires"))
+									.FixedWidth(75.0f)
+								)
+							]
 						]
+						+ SVerticalBox::Slot()
+							.AutoHeight()
+							.Padding(0.0f, 6.0f, 0.0f, 4.0f)
+							[
+								SNew(SSeparator)
+								.Orientation(Orient_Horizontal)
+							]
+						+ SVerticalBox::Slot()
+							.AutoHeight()
+							[
+								SNew(SHorizontalBox)
+								+ SHorizontalBox::Slot()
+								.AutoWidth()
+								[
+									SNew(SRichTextBlock)
+									.Text(LOCTEXT("CertificateMessage", "<RichTextBlock.TextHighlight>Note</>: The certificate in green will be used to sign the IPA."))
+									.TextStyle(FEditorStyle::Get(), "MessageLog")
+									.DecoratorStyleSet(&FEditorStyle::Get())
+									.AutoWrapText(true)
+								]
+								+ SHorizontalBox::Slot()
+								.FillWidth(1.0f)
+								.HAlign(HAlign_Right)
+								[
+									SNew(STextBlock)
+									.Text(LOCTEXT("ViewLabel", "View:"))
+								]
+
+								+ SHorizontalBox::Slot()
+									.AutoWidth()
+									.Padding(8.0f, 0.0f)
+									[
+										// all provisions hyper link
+										SNew(SHyperlink)
+										.OnNavigate(this, &FIOSTargetSettingsCustomization::HandleAllCertificatesHyperlinkNavigate, true)
+										.Text(LOCTEXT("AllCertificatesHyperLinkLabel", "All"))
+										.ToolTipText(LOCTEXT("AllCertificatesButtonTooltip", "View all certificates."))
+									]
+
+								+ SHorizontalBox::Slot()
+									.AutoWidth()
+									[
+										// valid provisions hyper link
+										SNew(SHyperlink)
+										.OnNavigate(this, &FIOSTargetSettingsCustomization::HandleAllCertificatesHyperlinkNavigate, false)
+										.Text(LOCTEXT("ValidCertificatesHyperlinkLabel", "Valid Only"))
+										.ToolTipText(LOCTEXT("ValidCertificatesHyperlinkTooltip", "View Valid certificates."))
+									]
+							]
+					]
 				]
 			+ SVerticalBox::Slot()
 			.AutoHeight()
@@ -696,6 +773,14 @@ void FIOSTargetSettingsCustomization::FindRequiredFiles()
 	IPPProcess->OnOutput().BindStatic(&OnOutput);
 	IPPProcess->Launch();
 	TickerHandle = FTicker::GetCoreTicker().AddTicker(FTickerDelegate::CreateRaw(this, &FIOSTargetSettingsCustomization::UpdateStatusDelegate), 1.0f);
+	if (ProvisionInfoSwitcher.IsValid())
+	{
+		ProvisionInfoSwitcher->SetActiveWidgetIndex(0);
+	}
+	if (CertificateInfoSwitcher.IsValid())
+	{
+		CertificateInfoSwitcher->SetActiveWidgetIndex(0);
+	}
 }
 
 FReply FIOSTargetSettingsCustomization::OnInstallProvisionClicked()
@@ -926,6 +1011,25 @@ void FIOSTargetSettingsCustomization::FilterLists()
 		}
 	}
 
+	if (ProvisionList.Num() > 0)
+	{
+		if (ProvisionInfoSwitcher.IsValid())
+		{
+			ProvisionInfoSwitcher->SetActiveWidgetIndex(2);
+		}
+		if (FilteredProvisionList.Num() == 0 && !bShowAllProvisions)
+		{
+			FilteredProvisionList.Append(ProvisionList);
+		}
+	}
+	else
+	{
+		if (ProvisionInfoSwitcher.IsValid())
+		{
+			ProvisionInfoSwitcher->SetActiveWidgetIndex(1);
+		}
+	}
+
 	for (int Index = 0; Index < CertificateList.Num(); ++Index)
 	{
 		if (SelectedCert.Contains(CertificateList[Index]->Name))
@@ -942,6 +1046,25 @@ void FIOSTargetSettingsCustomization::FilterLists()
 		}
 	}
 
+	if (CertificateList.Num() > 0)
+	{
+		if (CertificateInfoSwitcher.IsValid())
+		{
+			CertificateInfoSwitcher->SetActiveWidgetIndex(2);
+		}
+		if (FilteredCertificateList.Num() == 0 && !bShowAllCertificates)
+		{
+			FilteredCertificateList.Append(CertificateList);
+		}
+	}
+	else
+	{
+		if (CertificateInfoSwitcher.IsValid())
+		{
+			CertificateInfoSwitcher->SetActiveWidgetIndex(1);
+		}
+	}
+
 	CertificateListView->RequestListRefresh();
 	ProvisionListView->RequestListRefresh();
 }
diff --git a/Engine/Source/Developer/iOS/IOSPlatformEditor/Private/IOSTargetSettingsCustomization.h b/Engine/Source/Developer/iOS/IOSPlatformEditor/Private/IOSTargetSettingsCustomization.h
--- a/Engine/Source/Developer/iOS/IOSPlatformEditor/Private/IOSTargetSettingsCustomization.h
+++ b/Engine/Source/Developer/iOS/IOSPlatformEditor/Private/IOSTargetSettingsCustomization.h
@@ -71,9 +71,11 @@ private:
 	TArray<ProvisionPtr> ProvisionList;
 	TArray<ProvisionPtr> FilteredProvisionList;
 	TSharedPtr<SListView<ProvisionPtr> > ProvisionListView;
+	TSharedPtr<SWidgetSwitcher > ProvisionInfoSwitcher;
 	TArray<CertificatePtr> CertificateList;
 	TArray<CertificatePtr> FilteredCertificateList;
 	TSharedPtr<SListView<CertificatePtr> > CertificateListView;
+	TSharedPtr<SWidgetSwitcher > CertificateInfoSwitcher;
 
 	FString SelectedProvision;
 	FString SelectedFile;
