fix for not utilizing the game name if it contained [PROJECT_NAME]
UE-9047
#ios

diff --git a/Engine/Source/Developer/iOS/IOSPlatformEditor/Private/IOSTargetSettingsCustomization.cpp b/Engine/Source/Developer/iOS/IOSPlatformEditor/Private/IOSTargetSettingsCustomization.cpp
--- a/Engine/Source/Developer/iOS/IOSPlatformEditor/Private/IOSTargetSettingsCustomization.cpp
+++ b/Engine/Source/Developer/iOS/IOSPlatformEditor/Private/IOSTargetSettingsCustomization.cpp
@@ -759,14 +759,15 @@ void FIOSTargetSettingsCustomization::BuildImageRow(IDetailLayoutBuilder& Detail
 void FIOSTargetSettingsCustomization::FindRequiredFiles()
 {
 	const UIOSRuntimeSettings& Settings = *GetDefault<UIOSRuntimeSettings>();
+	FString BundleIdentifier = Settings.BundleIdentifier.Replace(TEXT("[PROJECT_NAME]"), FApp::GetGameName());
 #if PLATFORM_MAC
 	FString CmdExe = TEXT("/bin/sh");
 	FString ScriptPath = FPaths::ConvertRelativePathToFull(FPaths::EngineDir() / TEXT("Build/BatchFiles/Mac/RunMono.sh"));
 	FString IPPPath = FPaths::ConvertRelativePathToFull(FPaths::EngineDir() / TEXT("Binaries/DotNet/IOS/IPhonePackager.exe"));
-	FString CommandLine = FString::Printf(TEXT("\"%s\" \"%s\" certificates Engine -bundlename \"%s\""), *ScriptPath, *IPPPath, *(Settings.BundleIdentifier));
+	FString CommandLine = FString::Printf(TEXT("\"%s\" \"%s\" certificates Engine -bundlename \"%s\""), *ScriptPath, *IPPPath, *(BundleIdentifier));
 #else
 	FString CmdExe = FPaths::ConvertRelativePathToFull(FPaths::EngineDir() / TEXT("Binaries/DotNet/IOS/IPhonePackager.exe"));
-	FString CommandLine = FString::Printf(TEXT("certificates Engine -bundlename \"%s\""), *(Settings.BundleIdentifier));
+	FString CommandLine = FString::Printf(TEXT("certificates Engine -bundlename \"%s\""), *(BundleIdentifier));
 #endif
 	IPPProcess = MakeShareable(new FMonitoredProcess(CmdExe, CommandLine, true));
 	OutputMessage = TEXT("");
