fix for crashes when importing certificate/provision and/or modifying bundle identifier
UE-9214
#ios

diff --git a/Engine/Source/Developer/iOS/IOSPlatformEditor/Private/IOSTargetSettingsCustomization.cpp b/Engine/Source/Developer/iOS/IOSPlatformEditor/Private/IOSTargetSettingsCustomization.cpp
--- a/Engine/Source/Developer/iOS/IOSPlatformEditor/Private/IOSTargetSettingsCustomization.cpp
+++ b/Engine/Source/Developer/iOS/IOSPlatformEditor/Private/IOSTargetSettingsCustomization.cpp
@@ -2,13 +2,13 @@
 
 #include "IOSPlatformEditorPrivatePCH.h"
 #include "SWidgetSwitcher.h"
+#include "IDetailPropertyRow.h"
 #include "IOSTargetSettingsCustomization.h"
 #include "DetailLayoutBuilder.h"
 #include "DetailCategoryBuilder.h"
 #include "PropertyEditing.h"
 #include "DesktopPlatformModule.h"
 #include "MainFrame.h"
-#include "IDetailPropertyRow.h"
 
 #include "ScopedTransaction.h"
 #include "SExternalImageReference.h"
@@ -207,7 +207,7 @@ void FIOSTargetSettingsCustomization::BuildPListSection(IDetailLayoutBuilder& De
 	// Info.plist category
 	IDetailCategoryBuilder& ProvisionCategory = DetailLayout.EditCategory(TEXT("Mobile Provision"));
 	IDetailCategoryBuilder& AppManifestCategory = DetailLayout.EditCategory(TEXT("Info.plist"));
-	IDetailCategoryBuilder& BundleCategory = DetailLayout.EditCategory(TEXT("Bundle Information"));
+	IDetailCategoryBuilder& BundleCategory = DetailLayout.EditCategory(TEXT("BundleInformation"));
 	IDetailCategoryBuilder& OrientationCategory = DetailLayout.EditCategory(TEXT("Orientation"));
 	IDetailCategoryBuilder& RenderCategory = DetailLayout.EditCategory(TEXT("Rendering"));
 	IDetailCategoryBuilder& OSInfoCategory = DetailLayout.EditCategory(TEXT("OS Info"));
@@ -215,10 +215,6 @@ void FIOSTargetSettingsCustomization::BuildPListSection(IDetailLayoutBuilder& De
 	IDetailCategoryBuilder& BuildCategory = DetailLayout.EditCategory(TEXT("Build"));
 	IDetailCategoryBuilder& ExtraCategory = DetailLayout.EditCategory(TEXT("Extra PList Data"));
 
-	TSharedRef<SPlatformSetupMessage> PlatformSetupMessage = SNew(SPlatformSetupMessage, GameInfoPath)
-		.PlatformName(LOCTEXT("iOSPlatformName", "iOS"))
-		.OnSetupClicked(this, &FIOSTargetSettingsCustomization::CopySetupFilesIntoProject);
-
 /*	ProvisionCategory.AddCustomRow(TEXT("Certificate Request"), false)
 		.NameContent()
 		[
@@ -272,6 +268,19 @@ void FIOSTargetSettingsCustomization::BuildPListSection(IDetailLayoutBuilder& De
 						.AutoWrapText( true )
 					]
 				]
+				// importing a provision
+				+SWidgetSwitcher::Slot()
+					.HAlign(HAlign_Fill)
+					.VAlign(VAlign_Center)
+					[
+						SNew(SBorder)
+						.Padding(4)
+						[
+							SNew( STextBlock )
+							.Text( LOCTEXT( "ProvisionViewerImportingProvisions", "Importing Provision.  Please wait..." ) )
+							.AutoWrapText( true )
+						]
+					]
 				// no provisions found or no valid provisions
 				+SWidgetSwitcher::Slot()
 				.HAlign(HAlign_Fill)
@@ -379,6 +388,7 @@ void FIOSTargetSettingsCustomization::BuildPListSection(IDetailLayoutBuilder& De
 						.HAlign(HAlign_Center)
 						.VAlign(VAlign_Center)
 						.OnClicked(this, &FIOSTargetSettingsCustomization::OnInstallProvisionClicked)
+						.IsEnabled(this, &FIOSTargetSettingsCustomization::IsImportEnabled)
 						[
 							SNew(STextBlock)
 							.Text(FText::FromString("Import Provision"))
@@ -412,6 +422,19 @@ void FIOSTargetSettingsCustomization::BuildPListSection(IDetailLayoutBuilder& De
 						.AutoWrapText( true )
 					]
 				]
+				// importing certificate
+				+SWidgetSwitcher::Slot()
+					.HAlign(HAlign_Fill)
+					.VAlign(VAlign_Center)
+					[
+						SNew(SBorder)
+						.Padding(4)
+						[
+							SNew( STextBlock )
+							.Text( LOCTEXT( "CertificateViewerImportingCertificate", "Importing Certificate.  Please wait..." ) )
+							.AutoWrapText( true )
+						]
+					]
 				// no provisions found or no valid provisions
 				+SWidgetSwitcher::Slot()
 					.HAlign(HAlign_Fill)
@@ -520,6 +543,7 @@ void FIOSTargetSettingsCustomization::BuildPListSection(IDetailLayoutBuilder& De
 					.HAlign(HAlign_Center)
 					.VAlign(VAlign_Center)
 					.OnClicked(this, &FIOSTargetSettingsCustomization::OnInstallCertificateClicked)
+					.IsEnabled(this, &FIOSTargetSettingsCustomization::IsImportEnabled)
 					[
 						SNew(STextBlock)
 						.Text(FText::FromString("Import Certificate"))
@@ -549,22 +573,92 @@ void FIOSTargetSettingsCustomization::BuildPListSection(IDetailLayoutBuilder& De
 		 ]
 	 ];
 
-#define SETUP_NONROCKET_PROP(PropName, Category, Tip, DisabledTip) \
+	// Show properties that are gated by the plist being present and writable
+	RunningIPPProcess = false;
+
+#define SETUP_NONROCKET_PROP(PropName, Category) \
 	{ \
 		TSharedRef<IPropertyHandle> PropertyHandle = DetailLayout.GetProperty(GET_MEMBER_NAME_CHECKED(UIOSRuntimeSettings, PropName)); \
 		Category.AddProperty(PropertyHandle) \
 			.IsEnabled(!FRocketSupport::IsRocket()) \
-			.ToolTip(!FRocketSupport::IsRocket() ? Tip : DisabledTip); \
+			.ToolTip(!FRocketSupport::IsRocket() ? PropertyHandle->GetToolTipText() : FIOSTargetSettingsCustomizationConstants::DisabledTip); \
+	}
+
+#define SETUP_PLIST_PROP(PropName, Category) \
+	{ \
+		TSharedRef<IPropertyHandle> PropertyHandle = DetailLayout.GetProperty(GET_MEMBER_NAME_CHECKED(UIOSRuntimeSettings, PropName)); \
+		Category.AddProperty(PropertyHandle); \
 	}
 
-	SETUP_NONROCKET_PROP(bDevForArmV7, BuildCategory, LOCTEXT("DevForArmV7ToolTip", "Enable ArmV7 support? (this will be used if all type are unchecked)"), FIOSTargetSettingsCustomizationConstants::DisabledTip);
-	SETUP_NONROCKET_PROP(bDevForArm64, BuildCategory, LOCTEXT("DevForArm64ToolTip", "Enable Arm64 support?"), FIOSTargetSettingsCustomizationConstants::DisabledTip);
-	SETUP_NONROCKET_PROP(bDevForArmV7S, BuildCategory, LOCTEXT("DevForArmV7SToolTip", "Enable ArmV7s support?"), FIOSTargetSettingsCustomizationConstants::DisabledTip);
-	SETUP_NONROCKET_PROP(bShipForArmV7, BuildCategory, LOCTEXT("ShipForArmV7ToolTip", "Enable ArmV7 support? (this will be used if all type are unchecked)"), FIOSTargetSettingsCustomizationConstants::DisabledTip);
-	SETUP_NONROCKET_PROP(bShipForArm64, BuildCategory, LOCTEXT("ShipForArm64ToolTip", "Enable Arm64 support?"), FIOSTargetSettingsCustomizationConstants::DisabledTip);
-	SETUP_NONROCKET_PROP(bShipForArmV7S, BuildCategory, LOCTEXT("ShipForArmV7SToolTip", "Enable ArmV7s support?"), FIOSTargetSettingsCustomizationConstants::DisabledTip);
+#define SETUP_STATUS_PROP(PropName, Category) \
+	{ \
+		TSharedRef<IPropertyHandle> PropertyHandle = DetailLayout.GetProperty(GET_MEMBER_NAME_CHECKED(UIOSRuntimeSettings, PropName)); \
+		Category.AddProperty(PropertyHandle) \
+		.Visibility(EVisibility::Hidden); \
+		Category.AddCustomRow(LOCTEXT("BundleIdentifier", "BundleIdentifier"), false) \
+		.NameContent() \
+		[ \
+			SNew(SHorizontalBox) \
+			+ SHorizontalBox::Slot() \
+			.Padding(FMargin(0, 1, 0, 1)) \
+			.FillWidth(1.0f) \
+			[ \
+				SNew(STextBlock) \
+				.Text(LOCTEXT("BundleIdentifierLabel", "Bundle Identifier")) \
+				.Font(DetailLayout.GetDetailFont()) \
+			]\
+		] \
+		.ValueContent() \
+		.MinDesiredWidth( 0.0f ) \
+		.MaxDesiredWidth( 0.0f ) \
+		[ \
+			SNew(SHorizontalBox) \
+			+ SHorizontalBox::Slot() \
+			.FillWidth(1.0f) \
+			.HAlign(HAlign_Fill) \
+			[ \
+				SNew(SEditableTextBox) \
+				.IsEnabled(this, &FIOSTargetSettingsCustomization::IsImportEnabled) \
+				.Text(this, &FIOSTargetSettingsCustomization::GetBundleText, PropertyHandle) \
+				.Font(DetailLayout.GetDetailFont()) \
+				.SelectAllTextOnCommit( true ) \
+				.SelectAllTextWhenFocused( true ) \
+				.ClearKeyboardFocusOnCommit(false) \
+				.ToolTipText(PropertyHandle->GetToolTipText()) \
+				.OnTextCommitted(this, &FIOSTargetSettingsCustomization::OnBundleIdentifierChanged, PropertyHandle) \
+			] \
+		]; \
+	}
 
-	SETUP_NONROCKET_PROP(bSupportsMetalMRT, RenderCategory, LOCTEXT("MetalMRTToolTip", "Whether or not to add support for deferred rendering Metal API (requires IOS8 and A8 processors). This is very much a work in progress, be aware!"), FIOSTargetSettingsCustomizationConstants::DisabledTip);
+	const UIOSRuntimeSettings& Settings = *GetDefault<UIOSRuntimeSettings>();
+
+	SETUP_PLIST_PROP(BundleDisplayName, BundleCategory);
+	SETUP_PLIST_PROP(BundleName, BundleCategory);
+	SETUP_STATUS_PROP(BundleIdentifier, BundleCategory);
+	SETUP_PLIST_PROP(VersionInfo, BundleCategory);
+	SETUP_PLIST_PROP(bSupportsPortraitOrientation, OrientationCategory);
+	SETUP_PLIST_PROP(bSupportsUpsideDownOrientation, OrientationCategory);
+	SETUP_PLIST_PROP(bSupportsLandscapeLeftOrientation, OrientationCategory);
+	SETUP_PLIST_PROP(bSupportsLandscapeRightOrientation, OrientationCategory);
+	
+	SETUP_PLIST_PROP(bSupportsMetal, RenderCategory);
+	SETUP_PLIST_PROP(bSupportsOpenGLES2, RenderCategory);
+
+	SETUP_PLIST_PROP(bSupportsIPad, DeviceCategory);
+	SETUP_PLIST_PROP(bSupportsIPhone, DeviceCategory);
+
+	SETUP_PLIST_PROP(MinimumiOSVersion, OSInfoCategory);
+
+	SETUP_PLIST_PROP(AdditionalPlistData, ExtraCategory);
+
+	SETUP_NONROCKET_PROP(bDevForArmV7, BuildCategory);
+	SETUP_NONROCKET_PROP(bDevForArm64, BuildCategory);
+	SETUP_NONROCKET_PROP(bDevForArmV7S, BuildCategory);
+	SETUP_NONROCKET_PROP(bShipForArmV7, BuildCategory);
+	SETUP_NONROCKET_PROP(bShipForArm64, BuildCategory);
+	SETUP_NONROCKET_PROP(bShipForArmV7S, BuildCategory);
+
+	SETUP_NONROCKET_PROP(bSupportsMetalMRT, RenderCategory);
 
 #undef SETUP_NONROCKET_PROP
 }
@@ -744,6 +838,7 @@ void FIOSTargetSettingsCustomization::FindRequiredFiles()
 	{
 		CertificateInfoSwitcher->SetActiveWidgetIndex(0);
 	}
+	RunningIPPProcess = true;
 }
 
 FReply FIOSTargetSettingsCustomization::OnInstallProvisionClicked()
@@ -822,6 +917,11 @@ FReply FIOSTargetSettingsCustomization::OnInstallProvisionClicked()
 		IPPProcess->OnOutput().BindStatic(&OnOutput);
 		IPPProcess->Launch();
 		TickerHandle = FTicker::GetCoreTicker().AddTicker(FTickerDelegate::CreateRaw(this, &FIOSTargetSettingsCustomization::UpdateStatusDelegate), 10.0f);
+		if (ProvisionInfoSwitcher.IsValid())
+		{
+			ProvisionInfoSwitcher->SetActiveWidgetIndex(1);
+		}
+		RunningIPPProcess = true;
 	}
 
 	return FReply::Handled();
@@ -881,6 +981,11 @@ FReply FIOSTargetSettingsCustomization::OnInstallCertificateClicked()
 		IPPProcess->OnOutput().BindStatic(&OnOutput);
 		IPPProcess->Launch();
 		TickerHandle = FTicker::GetCoreTicker().AddTicker(FTickerDelegate::CreateRaw(this, &FIOSTargetSettingsCustomization::UpdateStatusDelegate), 10.0f);
+		if (CertificateInfoSwitcher.IsValid())
+		{
+			CertificateInfoSwitcher->SetActiveWidgetIndex(1);
+		}
+		RunningIPPProcess = true;
 	}
 
 	return FReply::Handled();
@@ -918,13 +1023,17 @@ const FSlateBrush* FIOSTargetSettingsCustomization::GetCertificateStatus() const
 
 bool FIOSTargetSettingsCustomization::UpdateStatusDelegate(float DeltaTime)
 {
-	if (IPPProcess->IsRunning())
+	if (IPPProcess.IsValid())
 	{
-		return true;
+		if (IPPProcess->IsRunning())
+		{
+			return true;
+		}
+		int RetCode = IPPProcess->GetReturnCode();
+		IPPProcess = NULL;
+		UpdateStatus();
 	}
-	int RetCode = IPPProcess->GetReturnCode();
-	IPPProcess = NULL;
-	UpdateStatus();
+	RunningIPPProcess = false;
 
 	return false;
 }
@@ -978,7 +1087,7 @@ void FIOSTargetSettingsCustomization::FilterLists()
 	{
 		if (ProvisionInfoSwitcher.IsValid())
 		{
-			ProvisionInfoSwitcher->SetActiveWidgetIndex(2);
+			ProvisionInfoSwitcher->SetActiveWidgetIndex(3);
 		}
 		if (FilteredProvisionList.Num() == 0 && !bShowAllProvisions)
 		{
@@ -989,7 +1098,7 @@ void FIOSTargetSettingsCustomization::FilterLists()
 	{
 		if (ProvisionInfoSwitcher.IsValid())
 		{
-			ProvisionInfoSwitcher->SetActiveWidgetIndex(1);
+			ProvisionInfoSwitcher->SetActiveWidgetIndex(2);
 		}
 	}
 
@@ -1013,7 +1122,7 @@ void FIOSTargetSettingsCustomization::FilterLists()
 	{
 		if (CertificateInfoSwitcher.IsValid())
 		{
-			CertificateInfoSwitcher->SetActiveWidgetIndex(2);
+			CertificateInfoSwitcher->SetActiveWidgetIndex(3);
 		}
 		if (FilteredCertificateList.Num() == 0 && !bShowAllCertificates)
 		{
@@ -1024,7 +1133,7 @@ void FIOSTargetSettingsCustomization::FilterLists()
 	{
 		if (CertificateInfoSwitcher.IsValid())
 		{
-			CertificateInfoSwitcher->SetActiveWidgetIndex(1);
+			CertificateInfoSwitcher->SetActiveWidgetIndex(2);
 		}
 	}
 
@@ -1032,6 +1141,29 @@ void FIOSTargetSettingsCustomization::FilterLists()
 	ProvisionListView->RequestListRefresh();
 }
 
+bool FIOSTargetSettingsCustomization::IsImportEnabled() const
+{
+	return !RunningIPPProcess.Get();
+}
+
+void FIOSTargetSettingsCustomization::OnBundleIdentifierChanged(const FText& NewText, ETextCommit::Type CommitType, TSharedRef<IPropertyHandle> InPropertyHandle)
+{
+	FText OutText;
+	InPropertyHandle->GetValueAsFormattedText(OutText);
+	if (OutText.ToString() != NewText.ToString())
+	{
+		InPropertyHandle->SetValueFromFormattedString( NewText.ToString() );
+		FindRequiredFiles();
+	}
+}
+
+FText FIOSTargetSettingsCustomization::GetBundleText(TSharedRef<IPropertyHandle> InPropertyHandle) const
+{
+	FText OutText;
+	InPropertyHandle->GetValueAsFormattedText(OutText);
+	return OutText;
+}
+
 //////////////////////////////////////////////////////////////////////////
 
-#undef LOCTEXT_NAMESPACE
\ No newline at end of file
+#undef LOCTEXT_NAMESPACE
diff --git a/Engine/Source/Developer/iOS/IOSPlatformEditor/Private/IOSTargetSettingsCustomization.h b/Engine/Source/Developer/iOS/IOSPlatformEditor/Private/IOSTargetSettingsCustomization.h
--- a/Engine/Source/Developer/iOS/IOSPlatformEditor/Private/IOSTargetSettingsCustomization.h
+++ b/Engine/Source/Developer/iOS/IOSPlatformEditor/Private/IOSTargetSettingsCustomization.h
@@ -76,6 +76,7 @@ private:
 	TArray<CertificatePtr> FilteredCertificateList;
 	TSharedPtr<SListView<CertificatePtr> > CertificateListView;
 	TSharedPtr<SWidgetSwitcher > CertificateInfoSwitcher;
+	TAttribute<bool> RunningIPPProcess;
 
 	FString SelectedProvision;
 	FString SelectedFile;
@@ -135,4 +136,13 @@ private:
 
 	// filter the lists based on the settings
 	void FilterLists();
+
+	// returns whether we are importing or not
+	bool IsImportEnabled() const;
+
+	// updates the text in the ini file and checks for a valid provision/certificate
+	void OnBundleIdentifierChanged(const FText& NewText, ETextCommit::Type, TSharedRef<IPropertyHandle> InPropertyHandle);
+
+	// 
+	FText GetBundleText(TSharedRef<IPropertyHandle> InPropertyHandle) const;
 };
diff --git a/Engine/Source/Runtime/IOS/IOSRuntimeSettings/Classes/IOSRuntimeSettings.h b/Engine/Source/Runtime/IOS/IOSRuntimeSettings/Classes/IOSRuntimeSettings.h
--- a/Engine/Source/Runtime/IOS/IOSRuntimeSettings/Classes/IOSRuntimeSettings.h
+++ b/Engine/Source/Runtime/IOS/IOSRuntimeSettings/Classes/IOSRuntimeSettings.h
@@ -169,19 +169,19 @@ public:
 	uint32 bSupportsLandscapeRightOrientation : 1;
 
 	// Specifies the the display name for the application. This will be displayed under the icon on the device.
-	UPROPERTY(GlobalConfig, EditAnywhere, Category = BundleInfo)
+	UPROPERTY(GlobalConfig, EditAnywhere, Category = BundleInformation)
 	FString BundleDisplayName;
 
 	// Specifies the the name of the application bundle. This is the short name for the application bundle.
-	UPROPERTY(GlobalConfig, EditAnywhere, Category = BundleInfo)
+	UPROPERTY(GlobalConfig, EditAnywhere, Category = BundleInformation)
 	FString BundleName;
 
 	// Specifies the bundle identifier for the application.
-	UPROPERTY(GlobalConfig, EditAnywhere, Category = BundleInfo)
+	UPROPERTY(GlobalConfig, EditAnywhere, Category = BundleInformation)
 	FString BundleIdentifier;
 
 	// Specifies the version for the application.
-	UPROPERTY(GlobalConfig, EditAnywhere, Category = BundleInfo)
+	UPROPERTY(GlobalConfig, EditAnywhere, Category = BundleInformation)
 	FString VersionInfo;
     
     /** Set the maximum frame rate to save on power consumption */
