- Proper fix for WITH_ENGINE = 0

diff --git a/Engine/Source/Developer/iOS/IOSTargetPlatform/Private/IOSTargetPlatform.cpp b/Engine/Source/Developer/iOS/IOSTargetPlatform/Private/IOSTargetPlatform.cpp
--- a/Engine/Source/Developer/iOS/IOSTargetPlatform/Private/IOSTargetPlatform.cpp
+++ b/Engine/Source/Developer/iOS/IOSTargetPlatform/Private/IOSTargetPlatform.cpp
@@ -295,8 +295,6 @@ bool FIOSTargetPlatform::HandleTicker(float DeltaTime )
 /* ITargetPlatform interface
  *****************************************************************************/
 
-#if WITH_ENGINE
-
 static bool SupportsES2()
 {
 	// default to supporting ES2
@@ -337,6 +335,28 @@ static bool CookASTC()
 	return bCookASTCTextures;
 }
 
+bool FIOSTargetPlatform::SupportsFeature( ETargetPlatformFeatures Feature ) const
+{
+	switch (Feature)
+	{
+		case ETargetPlatformFeatures::Packaging:
+			return true;
+			
+		case ETargetPlatformFeatures::LowQualityLightmaps:
+			return SupportsES2() || SupportsMetal();
+			
+		case ETargetPlatformFeatures::HighQualityLightmaps:
+			return SupportsMetalMRT();
+		default:
+			break;
+	}
+	
+	return TTargetPlatformBase<FIOSPlatformProperties>::SupportsFeature(Feature);
+}
+
+
+#if WITH_ENGINE
+
 
 void FIOSTargetPlatform::GetAllPossibleShaderFormats( TArray<FName>& OutFormats ) const
 {
@@ -438,22 +458,3 @@ FName FIOSTargetPlatform::GetWaveFormat( const class USoundWave* Wave ) const
 
 #endif // WITH_ENGINE
 
-bool FIOSTargetPlatform::SupportsFeature( ETargetPlatformFeatures Feature ) const
-{
-	switch (Feature)
-	{
-		case ETargetPlatformFeatures::Packaging:
-			return true;
-			
-		case ETargetPlatformFeatures::LowQualityLightmaps:
-			return SupportsES2() || SupportsMetal();
-			
-		case ETargetPlatformFeatures::HighQualityLightmaps:
-			return SupportsMetalMRT();
-		default:
-			break;
-	}
-	
-	return TTargetPlatformBase<FIOSPlatformProperties>::SupportsFeature(Feature);
-}
-
