SSH is now ready for use when building iOS from PC
UEPLAT-95
#ios

diff --git a/Engine/Source/Developer/iOS/IOSPlatformEditor/Private/IOSTargetSettingsCustomization.cpp b/Engine/Source/Developer/iOS/IOSPlatformEditor/Private/IOSTargetSettingsCustomization.cpp
--- a/Engine/Source/Developer/iOS/IOSPlatformEditor/Private/IOSTargetSettingsCustomization.cpp
+++ b/Engine/Source/Developer/iOS/IOSPlatformEditor/Private/IOSTargetSettingsCustomization.cpp
@@ -225,6 +225,13 @@ void FIOSTargetSettingsCustomization::UpdateStatus()
 	}
 }
 
+void FIOSTargetSettingsCustomization::UpdateStatus()
+{
+	// updated SSH key
+	const UIOSRuntimeSettings* Settings = GetDefault<UIOSRuntimeSettings>();
+	const_cast<UIOSRuntimeSettings*>(Settings)->PostInitProperties();
+}
+
 void FIOSTargetSettingsCustomization::BuildPListSection(IDetailLayoutBuilder& DetailLayout)
 {
 	// Info.plist category
@@ -689,20 +696,14 @@ void FIOSTargetSettingsCustomization::BuildPListSection(IDetailLayoutBuilder& De
 
 void FIOSTargetSettingsCustomization::BuildRemoteBuildingSection(IDetailLayoutBuilder& DetailLayout)
 {
+#if PLATFORM_WINDOWS
 	IDetailCategoryBuilder& BuildCategory = DetailLayout.EditCategory(TEXT("Build"));
 
 	// Sub group we wish to add remote building options to.
 	FText RemoteBuildingGroupName = LOCTEXT("RemoteBuildingGroupName", "Remote Build Options");
 	IDetailGroup& RemoteBuildingGroup = BuildCategory.AddGroup(*RemoteBuildingGroupName.ToString(), RemoteBuildingGroupName, false);
 
-	// Remote Server Name Property
-	TSharedRef<IPropertyHandle> RemoteServerNamePropertyHandle = DetailLayout.GetProperty(GET_MEMBER_NAME_CHECKED(UIOSRuntimeSettings, RemoteServerName));
-	IDetailPropertyRow& RemoteServerNamePropertyRow = RemoteBuildingGroup.AddPropertyRow(RemoteServerNamePropertyHandle);
-	RemoteServerNamePropertyRow
-		.IsEnabled(!FRocketSupport::IsRocket())
-		.ToolTip(!FRocketSupport::IsRocket() ? LOCTEXT("RemoteServerNameToolTip", "The name or ip address of the remote mac which will be used to build IOS") : FIOSTargetSettingsCustomizationConstants::DisabledTip);
 
-	
 	// Add Use RSync Property
 	TSharedRef<IPropertyHandle> UseRSyncPropertyHandle = DetailLayout.GetProperty(GET_MEMBER_NAME_CHECKED(UIOSRuntimeSettings, bUseRSync));
 	IDetailPropertyRow& UseRSyncPropertRow = RemoteBuildingGroup.AddPropertyRow(UseRSyncPropertyHandle);
@@ -710,13 +711,87 @@ void FIOSTargetSettingsCustomization::BuildRemoteBuildingSection(IDetailLayoutBu
 		.IsEnabled(!FRocketSupport::IsRocket())
 		.ToolTip(!FRocketSupport::IsRocket() ? LOCTEXT("UseRSyncToolTip", "Use RSync instead of RPCUtility") : FIOSTargetSettingsCustomizationConstants::DisabledTip);
 
+
+	// Remote Server Name Property
+	TSharedRef<IPropertyHandle> RemoteServerNamePropertyHandle = DetailLayout.GetProperty(GET_MEMBER_NAME_CHECKED(UIOSRuntimeSettings, RemoteServerName));
+	IDetailPropertyRow& RemoteServerNamePropertyRow = RemoteBuildingGroup.AddPropertyRow(RemoteServerNamePropertyHandle);
+	RemoteServerNamePropertyRow
+		.IsEnabled(!FRocketSupport::IsRocket())
+		.ToolTip(!FRocketSupport::IsRocket() ? LOCTEXT("RemoteServerNameToolTip", "The name or ip address of the remote mac which will be used to build IOS") : FIOSTargetSettingsCustomizationConstants::DisabledTip)
+		.CustomWidget()
+		.NameContent()
+		[
+			SNew(SHorizontalBox)
+			+ SHorizontalBox::Slot()
+			.Padding(FMargin(0, 1, 0, 1))
+			.FillWidth(1.0f)
+			[
+				SNew(STextBlock)
+				.Text(LOCTEXT("RemoteServerNameLabel", "Remote Server Name"))
+				.Font(DetailLayout.GetDetailFont())
+			]
+		]
+		.ValueContent()
+		.MinDesiredWidth(0.0f)
+		.MaxDesiredWidth(0.0f)
+		[
+			SNew(SHorizontalBox)
+			+ SHorizontalBox::Slot()
+			.FillWidth(1.0f)
+			.HAlign(HAlign_Fill)
+			[
+				SNew(SEditableTextBox)
+				.IsEnabled(this, &FIOSTargetSettingsCustomization::IsImportEnabled)
+				.Text(this, &FIOSTargetSettingsCustomization::GetBundleText, RemoteServerNamePropertyHandle)
+				.Font(DetailLayout.GetDetailFont())
+				.SelectAllTextOnCommit(true)
+				.SelectAllTextWhenFocused(true)
+				.ClearKeyboardFocusOnCommit(false)
+				.ToolTipText(RemoteServerNamePropertyHandle->GetToolTipText())
+				.OnTextCommitted(this, &FIOSTargetSettingsCustomization::OnRemoteServerChanged, RemoteServerNamePropertyHandle)
+			]
+		];
+
 	
 	// Add RSync Username Property
 	TSharedRef<IPropertyHandle> RSyncUsernamePropertyHandle = DetailLayout.GetProperty(GET_MEMBER_NAME_CHECKED(UIOSRuntimeSettings, RSyncUsername));
 	IDetailPropertyRow& RSyncUsernamePropertyRow = RemoteBuildingGroup.AddPropertyRow(RSyncUsernamePropertyHandle);
 	RSyncUsernamePropertyRow
 		.IsEnabled(!FRocketSupport::IsRocket())
-		.ToolTip(!FRocketSupport::IsRocket() ? LOCTEXT("RSyncUsernameToolTip", "The username of the mac user that matches the specified SSH Key.") : FIOSTargetSettingsCustomizationConstants::DisabledTip);
+		.ToolTip(!FRocketSupport::IsRocket() ? LOCTEXT("RSyncUsernameToolTip", "The username of the mac user that matches the specified SSH Key.") : FIOSTargetSettingsCustomizationConstants::DisabledTip)
+		.CustomWidget()
+			.NameContent()
+			[
+				SNew(SHorizontalBox)
+				+ SHorizontalBox::Slot()
+				.Padding(FMargin(0, 1, 0, 1))
+				.FillWidth(1.0f)
+				[
+					SNew(STextBlock)
+					.Text(LOCTEXT("RSyncUserNameLabel", "RSync User Name"))
+					.Font(DetailLayout.GetDetailFont())
+				]
+			]
+		.ValueContent()
+			.MinDesiredWidth(0.0f)
+			.MaxDesiredWidth(0.0f)
+			[
+				SNew(SHorizontalBox)
+				+ SHorizontalBox::Slot()
+				.FillWidth(1.0f)
+				.HAlign(HAlign_Fill)
+				[
+					SNew(SEditableTextBox)
+					.IsEnabled(this, &FIOSTargetSettingsCustomization::IsImportEnabled)
+					.Text(this, &FIOSTargetSettingsCustomization::GetBundleText, RSyncUsernamePropertyHandle)
+					.Font(DetailLayout.GetDetailFont())
+					.SelectAllTextOnCommit(true)
+					.SelectAllTextWhenFocused(true)
+					.ClearKeyboardFocusOnCommit(false)
+					.ToolTipText(RSyncUsernamePropertyHandle->GetToolTipText())
+					.OnTextCommitted(this, &FIOSTargetSettingsCustomization::OnRemoteServerChanged, RSyncUsernamePropertyHandle)
+				]
+			];
 
 
 	// Add existing SSH path label.
@@ -733,6 +808,36 @@ void FIOSTargetSettingsCustomization::BuildRemoteBuildingSection(IDetailLayoutBu
 	SSHPrivateKeyOverridePathPropertyRow
 		.IsEnabled(!FRocketSupport::IsRocket())
 		.ToolTip(!FRocketSupport::IsRocket() ? LOCTEXT("SSHPrivateKeyOverridePathToolTip", "Override the existing SSH Private Key with one from a specified location.") : FIOSTargetSettingsCustomizationConstants::DisabledTip);
+
+	// Add a generate key button
+	RemoteBuildingGroup.AddWidgetRow()
+		.WholeRowWidget
+		.MinDesiredWidth(0.f)
+		.MaxDesiredWidth(0.f)
+		.HAlign(HAlign_Fill)
+		[
+			SNew(SVerticalBox)
+			+ SVerticalBox::Slot()
+				.AutoHeight()
+				[
+					SNew(SHorizontalBox)
+					+ SHorizontalBox::Slot()
+					.Padding(FMargin(0, 5, 0, 10))
+					.AutoWidth()
+					[
+						SNew(SButton)
+						.HAlign(HAlign_Center)
+						.VAlign(VAlign_Center)
+						.OnClicked(this, &FIOSTargetSettingsCustomization::OnGenerateSSHKey)
+						.IsEnabled(this, &FIOSTargetSettingsCustomization::IsImportEnabled)
+						[
+							SNew(STextBlock)
+							.Text(FText::FromString("Generate SSH Key"))
+						]
+					]
+				]
+		];
+#endif
 }
 
 
@@ -1020,6 +1125,45 @@ FReply FIOSTargetSettingsCustomization::OnCertificateRequestClicked()
 	return FReply::Handled();
 }
 
+FReply FIOSTargetSettingsCustomization::OnGenerateSSHKey()
+{
+	// see if the key is already generated
+	const UIOSRuntimeSettings& Settings = *GetDefault<UIOSRuntimeSettings>();
+	TCHAR Path[4096];
+	FPlatformMisc::GetEnvironmentVariable(TEXT("APPDATA"), Path, ARRAY_COUNT(Path));
+	FString Destination = FString::Printf(TEXT("%s\\Unreal Engine\\UnrealBuildTool\\SSHKeys\\%s\\%s\\RemoteToolChainPrivate.key"), Path, *(Settings.RemoteServerName), *(Settings.RSyncUsername));
+	if (FPaths::FileExists(Destination))
+	{
+		FString MessagePrompt = FString::Printf(TEXT("An SSH Key already exists.  Do you want to replace this key?"));
+		if (FPlatformMisc::MessageBoxExt(EAppMsgType::OkCancel, *MessagePrompt, TEXT("Key Exists")) == EAppReturnType::Cancel)
+		{
+			return FReply::Handled();
+		}
+	}
+
+	FString CmdExe = FPaths::ConvertRelativePathToFull(FPaths::EngineDir() / TEXT("Build/BatchFiles/MakeAndInstallSSHKey.bat"));
+	TCHAR ProgramPath[4096];
+	FPlatformMisc::GetEnvironmentVariable(TEXT("PROGRAMFILES(X86)"), ProgramPath, ARRAY_COUNT(ProgramPath));
+	FString CygwinPath = TEXT("/cygdrive/") + FString(Path).Replace(TEXT(":"), TEXT("")).Replace(TEXT("\\"), TEXT("/"));
+	FString EnginePath = FPaths::EngineDir();
+	FString CommandLine = FString::Printf(TEXT("\"%s\\DeltaCopy\\ssh.exe\" \"%s\\DeltaCopy\\rsync.exe\" %s %s \"%s\" \"%s\" \"%s\""),
+		ProgramPath,
+		ProgramPath,
+		*(Settings.RSyncUsername),
+		*(Settings.RemoteServerName),
+		Path,
+		*CygwinPath,
+		*EnginePath);
+
+	OutputMessage = TEXT("");
+	IPPProcess = MakeShareable(new FMonitoredProcess(CmdExe, CommandLine, false, false));
+	IPPProcess->Launch();
+	TickerHandle = FTicker::GetCoreTicker().AddTicker(FTickerDelegate::CreateRaw(this, &FIOSTargetSettingsCustomization::UpdateStatusDelegate), 10.0f);
+	RunningIPPProcess = true;
+
+	return FReply::Handled();
+}
+
 const FSlateBrush* FIOSTargetSettingsCustomization::GetProvisionStatus() const
 {
 	if( bProvisionInstalled )
@@ -1055,6 +1199,7 @@ bool FIOSTargetSettingsCustomization::UpdateStatusDelegate(float DeltaTime)
 		int RetCode = IPPProcess->GetReturnCode();
 		IPPProcess = NULL;
 		UpdateStatus();
+		UpdateSSHStatus();
 	}
 	RunningIPPProcess = false;
 
@@ -1180,6 +1325,18 @@ void FIOSTargetSettingsCustomization::OnBundleIdentifierChanged(const FText& New
 	}
 }
 
+void FIOSTargetSettingsCustomization::OnRemoteServerChanged(const FText& NewText, ETextCommit::Type CommitType, TSharedRef<IPropertyHandle> InPropertyHandle)
+{
+	FText OutText;
+	InPropertyHandle->GetValueAsFormattedText(OutText);
+	if (OutText.ToString() != NewText.ToString())
+	{
+		InPropertyHandle->SetValueFromFormattedString(NewText.ToString());
+		OutputMessage = TEXT("");
+		UpdateSSHStatus();
+	}
+}
+
 FText FIOSTargetSettingsCustomization::GetBundleText(TSharedRef<IPropertyHandle> InPropertyHandle) const
 {
 	FText OutText;
diff --git a/Engine/Source/Developer/iOS/IOSPlatformEditor/Private/IOSTargetSettingsCustomization.h b/Engine/Source/Developer/iOS/IOSPlatformEditor/Private/IOSTargetSettingsCustomization.h
--- a/Engine/Source/Developer/iOS/IOSPlatformEditor/Private/IOSTargetSettingsCustomization.h
+++ b/Engine/Source/Developer/iOS/IOSPlatformEditor/Private/IOSTargetSettingsCustomization.h
@@ -107,6 +107,9 @@ private:
 	// certificate request
 	FReply OnCertificateRequestClicked();
 
+	// ssh key request
+	FReply OnGenerateSSHKey();
+
 	// Get the image to display for the provision status
 	const FSlateBrush* GetProvisionStatus() const;
 
@@ -119,6 +122,9 @@ private:
 	// Update the provision status
 	void UpdateStatus();
 
+	// Update the ssh key status
+	void UpdateSSHStatus();
+
 	// status tick delay
 	bool UpdateStatusDelegate(float delay);
 
@@ -143,6 +149,9 @@ private:
 	// updates the text in the ini file and checks for a valid provision/certificate
 	void OnBundleIdentifierChanged(const FText& NewText, ETextCommit::Type, TSharedRef<IPropertyHandle> InPropertyHandle);
 
+	// updates the text in the ini file and checks for a valid provision/certificate
+	void OnRemoteServerChanged(const FText& NewText, ETextCommit::Type, TSharedRef<IPropertyHandle> InPropertyHandle);
+
 	// 
 	FText GetBundleText(TSharedRef<IPropertyHandle> InPropertyHandle) const;
 };
diff --git a/Engine/Source/Programs/IOS/iPhonePackager/CompileTime.cs b/Engine/Source/Programs/IOS/iPhonePackager/CompileTime.cs
--- a/Engine/Source/Programs/IOS/iPhonePackager/CompileTime.cs
+++ b/Engine/Source/Programs/IOS/iPhonePackager/CompileTime.cs
@@ -37,6 +37,7 @@ namespace iPhonePackager
 		/// The file name (no path) of the temporary mobile provision that will be placed on the remote mac for use in makeapp
 		/// </summary>
 		private static string MacMobileProvisionFilename;
+		private static string MacSigningIdentityFilename;
 
 		private static string MacName = "";
 		private static string MacStagingRootDir = "";
@@ -125,10 +126,14 @@ namespace iPhonePackager
 
 			// generate the directories to recursively copy into later on
 			MacStagingRootDir = string.Format("{0}/{1}/IOS", iPhone_SigningDevRootMac, GameBranchPath);
+			MacStagingRootDir = MacStagingRootDir.Replace("//", "/");
 			MacBinariesDir = string.Format("{0}/{1}/IOS", iPhone_SigningDevRootMac, GameBranchPath);
+			MacBinariesDir = MacBinariesDir.Replace("//", "/");
 			MacXcodeStagingDir = string.Format("{0}/{1}/IOS/XcodeSupportFiles", iPhone_SigningDevRootMac, GameBranchPath);
+			MacXcodeStagingDir = MacXcodeStagingDir.Replace("//", "/");
 
 			MacMobileProvisionFilename = MachineName + "_UE4Temp.mobileprovision";
+			MacSigningIdentityFilename = MachineName + "_UE4Temp.p12";
 
 			CurrentBaseXCodeCommandLine = GetBaseXcodeCommandline();
 		}
@@ -265,6 +270,13 @@ namespace iPhonePackager
 			}
 			FileOperations.CopyRequiredFile(ProjectFile, Path.Combine(Config.PCXcodeStagingDir, @"project.pbxproj.datecheck"));
 			
+			// copy the signing certificate over
+			// export the signing certificate to a file
+			MobileProvision Provision = MobileProvisionParser.ParseFile(ProvisionWithPrefix);
+			var Certificate = CodeSignatureBuilder.FindCertificate(Provision);
+			byte[] Data = Certificate.Export(System.Security.Cryptography.X509Certificates.X509ContentType.Cert);
+			File.WriteAllBytes(Path.Combine(Config.PCXcodeStagingDir, MacSigningIdentityFilename), Data);
+
 			// needs Mac line endings so it can be executed
 			string SrcPath = @"..\..\..\Build\IOS\XcodeSupportFiles\prepackage.sh";
 			string DestPath = Path.Combine(Config.PCXcodeStagingDir, @"prepackage.sh");
@@ -284,6 +296,9 @@ namespace iPhonePackager
 			string CommandLine = "";
 			string WorkingFolder = "";
 			string DisplayCommandLine = "";
+			string TempKeychain = "$HOME/Library/Keychains/UE4TempKeychain.keychain";
+			string Certificate = "XcodeSupportFiles/" + MacSigningIdentityFilename;
+			string LoginKeychain = "$HOME/Library/Keychains/login.keychain";
 
 			switch (RPCCommand.ToLowerInvariant())
 			{
@@ -291,7 +306,7 @@ namespace iPhonePackager
 				Program.Log( " ... deleting staging files on the Mac" );
 				DisplayCommandLine = "rm -rf Payload";
 				CommandLine = "\"" + MacStagingRootDir + "\" " + DisplayCommandLine;
-				WorkingFolder = MacStagingRootDir;
+				WorkingFolder = "\"" + MacStagingRootDir + "\"";
 				break;
 
 			case "ensureprovisiondirexists":
@@ -300,7 +315,7 @@ namespace iPhonePackager
 				DisplayCommandLine = String.Format("mkdir -p ~/Library/MobileDevice/Provisioning\\ Profiles");
 
 				CommandLine = "\"" + MacXcodeStagingDir + "\" " + DisplayCommandLine;
-				WorkingFolder = MacXcodeStagingDir;
+				WorkingFolder = "\"" + MacXcodeStagingDir + "\"";
 				break;
 
 			case "installprovision":
@@ -310,14 +325,14 @@ namespace iPhonePackager
 				DisplayCommandLine = String.Format("cp -f {0} ~/Library/MobileDevice/Provisioning\\ Profiles", MacMobileProvisionFilename);
 
 				CommandLine = "\"" + MacXcodeStagingDir + "\" " + DisplayCommandLine;
-				WorkingFolder = MacXcodeStagingDir;
+				WorkingFolder = "\"" + MacXcodeStagingDir + "\"";
 				break;
 
 			case "removeprovision":
 				Program.Log(" ... removing .mobileprovision");
 				DisplayCommandLine = String.Format("rm -f ~/Library/MobileDevice/Provisioning\\ Profiles/{0}", MacMobileProvisionFilename);
 				CommandLine = "\"" + MacXcodeStagingDir + "\" " + DisplayCommandLine;
-				WorkingFolder = MacXcodeStagingDir;
+				WorkingFolder = "\"" + MacXcodeStagingDir + "\"";
 				break;
 
 			case "setexec": 
@@ -325,14 +340,14 @@ namespace iPhonePackager
 				Program.Log(" ... setting executable bit");
 				DisplayCommandLine = "chmod a+x \'" + RemoteExecutablePath + "\'";
 				CommandLine = "\"" + MacStagingRootDir + "\" " + DisplayCommandLine;
-				WorkingFolder = MacStagingRootDir;
+				WorkingFolder = "\"" + MacStagingRootDir + "\"";
 				break;
 
 			case "prepackage":
 				Program.Log(" ... running prepackage script remotely ");
 				DisplayCommandLine = String.Format("sh prepackage.sh {0} IOS {1} {2}", Program.GameName, Program.GameConfiguration, Program.Architecture);
 				CommandLine = "\"" + MacXcodeStagingDir + "\" " + DisplayCommandLine;
-				WorkingFolder = MacXcodeStagingDir;
+				WorkingFolder = "\"" + MacXcodeStagingDir + "\"";
 				break;
 
 			case "makeapp":
@@ -343,18 +358,34 @@ namespace iPhonePackager
 				WorkingFolder = "\"" + MacXcodeStagingDir + "/..\"";
 				break;
 
+			case "createkeychain":
+				Program.Log(" ... creating temporary key chain with signing certificate");
+				Program.Log("  Using signing identity '{0}'", Config.CodeSigningIdentity);
+				DisplayCommandLine = "security create-keychain -p \"\" \"" + TempKeychain + "\" && security list-keychains -s \"" + TempKeychain + "\" && security -v unlock-keychain -p \"\" \"" + TempKeychain + "\" && security import " + Certificate + " -k \"" + TempKeychain + "\" -P \"\" -A -t cert";
+				CommandLine = "\"" + MacXcodeStagingDir + "/..\" " + DisplayCommandLine;
+				WorkingFolder = "\"" + MacXcodeStagingDir + "/..\"";
+				break;
+
+			case "deletekeychain":
+				Program.Log(" ... remove temporary key chain");
+				Program.Log("  Using signing identity '{0}'", Config.CodeSigningIdentity);
+				DisplayCommandLine = "security list-keychains -s \"" + LoginKeychain + "\" && security delete-keychain \"" + TempKeychain + "\"";
+				CommandLine = "\"" + MacXcodeStagingDir + "/..\" " + DisplayCommandLine;
+				WorkingFolder = "\"" + MacXcodeStagingDir + "/..\"";
+				break;
+
 			case "validation":
 				Program.Log( " ... validating distribution package" );
 				DisplayCommandLine = XcodeDeveloperDir + "Platforms/iPhoneOS.platform/Developer/usr/bin/Validation " + RemoteAppDirectory;
 				CommandLine = "\"" + MacStagingRootDir + "\" " + DisplayCommandLine;
-				WorkingFolder = MacStagingRootDir;
+				WorkingFolder = "\"" + MacStagingRootDir + "\"";
 				break;
 
 			case "deleteipa":
 				Program.Log(" ... deleting IPA on Mac");
 				DisplayCommandLine = "rm -f " + Config.IPAFilenameOnMac;
 				CommandLine = "\"" + MacStagingRootDir + "\" " + DisplayCommandLine;
-				WorkingFolder = MacStagingRootDir;
+				WorkingFolder = "\"" + MacStagingRootDir + "\"";
 				break;
 
 			case "kill":
@@ -368,14 +399,14 @@ namespace iPhonePackager
 				Program.Log( " ... stripping" );
 				DisplayCommandLine = XcodeDeveloperDir + "Platforms/iPhoneOS.platform/Developer/usr/bin/strip '" + RemoteExecutablePath + "'";
 				CommandLine = "\"" + MacStagingRootDir + "\" " + DisplayCommandLine;
-				WorkingFolder = MacStagingRootDir;
+				WorkingFolder = "\"" + MacStagingRootDir + "\"";
 				break;
 
 			case "resign":
 				Program.Log("... resigning");
 				DisplayCommandLine = "bash -c '" + "chmod a+x ResignScript" + ";" + "./ResignScript" + "'";
 				CommandLine = "\"" + MacStagingRootDir + "\" " + DisplayCommandLine;
-				WorkingFolder = MacStagingRootDir;
+				WorkingFolder = "\"" + MacStagingRootDir + "\"";
 				break;
 
 			case "zip":
@@ -392,7 +423,7 @@ namespace iPhonePackager
 					dSYMName);
 
 				CommandLine = "\"" + MacStagingRootDir + "\" " + DisplayCommandLine;
-				WorkingFolder = MacStagingRootDir;
+				WorkingFolder = "\"" + MacStagingRootDir + "\"";
 				break;
 
 			case "gendsym":
@@ -403,7 +434,7 @@ namespace iPhonePackager
 				DisplayCommandLine = String.Format("dsymutil -o {0} {1}", dSYMPath, ExePath);
 
 				CommandLine = "\"" + MacStagingRootDir + "\"" + DisplayCommandLine;
-				WorkingFolder = MacStagingRootDir;
+				WorkingFolder = "\"" + MacStagingRootDir + "\"";
 				break;
 
 			default:
@@ -498,7 +529,15 @@ namespace iPhonePackager
 
 			// sign the exe, etc...
 			CompileTime.ExecuteRemoteCommand("PrePackage");
+			if (!Config.bUseRPCUtil)
+			{
+				CompileTime.ExecuteRemoteCommand("CreateKeyChain");
+			}
 			CompileTime.ExecuteRemoteCommand("MakeApp");
+			if (!Config.bUseRPCUtil)
+			{
+				CompileTime.ExecuteRemoteCommand("DeleteKeyChain");
+			}
 
 			Program.Log(String.Format("Finished creating .app directory on Mac (took {0:0.00} s)",
 				(DateTime.Now - StartTime).TotalSeconds));
diff --git a/Engine/Source/Programs/IOS/iPhonePackager/SSHCommandHelper.cs b/Engine/Source/Programs/IOS/iPhonePackager/SSHCommandHelper.cs
--- a/Engine/Source/Programs/IOS/iPhonePackager/SSHCommandHelper.cs
+++ b/Engine/Source/Programs/IOS/iPhonePackager/SSHCommandHelper.cs
@@ -245,6 +245,7 @@ namespace iPhonePackager
 		static public Hashtable UploadFile(string RemoteMac, string LocalPath, string RemotePath)
 		{
 			string RemoteDir = Path.GetDirectoryName(RemotePath).Replace("\\", "/");
+			RemoteDir = RemoteDir.Replace(" ", "\\ ");
 			string RemoteFilename = Path.GetFileName(RemotePath);
 
 			// get the executable dir for SSH, so Rsync can call it easily
@@ -259,7 +260,7 @@ namespace iPhonePackager
 			// make simple rsync commandline to send a file
 			RsyncProcess.StartInfo.FileName = RSyncExePath;
 			RsyncProcess.StartInfo.Arguments = string.Format(
-				"-zae \"ssh {0}\" --rsync-path=\"mkdir -p {1} && rsync\" \"{2}\" {3}@{4}:\"{1}/{5}\"",
+				"-zae \"ssh {0}\" --rsync-path=\"mkdir -p {1} && rsync\" '{2}' {3}@{4}:'{1}/{5}'",
 				SSHAuthentication,
 				RemoteDir,
 				ConvertPathToCygwin(LocalPath),
@@ -300,6 +301,8 @@ namespace iPhonePackager
 		{
 			// get the executable dir for SSH, so Rsync can call it easily
 			string ExeDir = Path.GetDirectoryName(SSHExePath);
+			string RemoteDir = RemotePath.Replace("\\", "/");
+			RemoteDir = RemoteDir.Replace(" ", "\\ ");
 
 			Process RsyncProcess = new Process();
 			if (ExeDir != "")
@@ -313,12 +316,12 @@ namespace iPhonePackager
 			// make simple rsync commandline to send a file
 			RsyncProcess.StartInfo.FileName = RSyncExePath;
 			RsyncProcess.StartInfo.Arguments = string.Format(
-				"-zae \"ssh {0}\" {2}@{3}:\"{4}\" \"{1}\"",
+				"-zae \"ssh {0}\" {2}@{3}:'{4}' \"{1}\"",
 				SSHAuthentication,
 				ConvertPathToCygwin(LocalPath),
 				SSHUser,
 				RemoteMac,
-				RemotePath
+				RemoteDir
 				);
 
 			SSHOutputMap[RsyncProcess] = new StringBuilder("");
diff --git a/Engine/Source/Programs/UnrealBuildTool/Configuration/UEBuildPlatform.cs b/Engine/Source/Programs/UnrealBuildTool/Configuration/UEBuildPlatform.cs
--- a/Engine/Source/Programs/UnrealBuildTool/Configuration/UEBuildPlatform.cs
+++ b/Engine/Source/Programs/UnrealBuildTool/Configuration/UEBuildPlatform.cs
@@ -948,6 +948,7 @@ namespace UnrealBuildTool
 				ProjIni.GetBool(Section, Key, out Project);
 				if (Default != Project)
 				{
+					Console.WriteLine(Key + " is not set to default. (" + Default + " vs. " + Project + ")");
 					return false;
 				}
 			}
@@ -960,6 +961,7 @@ namespace UnrealBuildTool
 				ProjIni.GetInt32(Section, Key, out Project);
 				if (Default != Project)
 				{
+					Console.WriteLine(Key + " is not set to default. (" + Default + " vs. " + Project + ")");
 					return false;
 				}
 			}
@@ -972,6 +974,7 @@ namespace UnrealBuildTool
 				ProjIni.GetString(Section, Key, out Project);
 				if (Default != Project)
 				{
+					Console.WriteLine(Key + " is not set to default. (" + Default + " vs. " + Project + ")");
 					return false;
 				}
 			}
diff --git a/Engine/Source/Programs/UnrealBuildTool/System/RPCUtilHelper.cs b/Engine/Source/Programs/UnrealBuildTool/System/RPCUtilHelper.cs
--- a/Engine/Source/Programs/UnrealBuildTool/System/RPCUtilHelper.cs
+++ b/Engine/Source/Programs/UnrealBuildTool/System/RPCUtilHelper.cs
@@ -115,8 +115,7 @@ namespace UnrealBuildTool
 						string[] Commands = new string[] 
 						{
 							"+\"%s\"",
-							"vm_stat | grep Statistics | awk '{print $8}'",
-							"vm_stat | grep 'Pages free:' | awk '{print $3}'",
+							"sysctl -a | grep hw.memsize | awk '{print $2}'",
 							"sysctl -a | grep hw.logicalcpu: | awk '{print $2}'",
 						};
 						Hashtable Results = Command("/", "date", string.Join(" && ", Commands), null);
@@ -139,12 +138,12 @@ namespace UnrealBuildTool
 						TimeDifferenceFromRemote = DateTime.UtcNow - RemoteTimebase;
 
 						// now figure out max number of commands to run at once
-						int PageSize = int.Parse(Lines[1]);
-						Int64 AvailableMem = Int64.Parse(Lines[2].Replace(".", "")) * PageSize;
+//						int PageSize = int.Parse(Lines[1]);
+						Int64 AvailableMem = Int64.Parse(Lines[1].Replace(".", ""));// *PageSize;
 						int NumProcesses = (int)Math.Max(1, AvailableMem / (RemoteToolChain.MemoryPerCompileMB * 1024 * 1024));
 
 						// now, combine that with actual number of cores
-						int NumCores = int.Parse(Lines[3]);
+						int NumCores = int.Parse(Lines[2]);
 						MaxRemoteCommandsAllowed = Math.Min(NumProcesses, NumCores);
 
 						Console.WriteLine("Remote time is {0}, difference is {1}", RemoteTimebase.ToString(), TimeDifferenceFromRemote.ToString());
diff --git a/Engine/Source/Programs/UnrealBuildTool/ToolChain/RemoteToolChain.cs b/Engine/Source/Programs/UnrealBuildTool/ToolChain/RemoteToolChain.cs
--- a/Engine/Source/Programs/UnrealBuildTool/ToolChain/RemoteToolChain.cs
+++ b/Engine/Source/Programs/UnrealBuildTool/ToolChain/RemoteToolChain.cs
@@ -33,8 +33,11 @@ namespace UnrealBuildTool
 		/** Keep a list of remote files that are potentially copied from local to remote */
 		private static Dictionary<FileItem, FileItem> CachedRemoteFileItems = new Dictionary<FileItem, FileItem>();
 
-		/** The path (on the Mac) to the your particular development directory, where files will be copied to from the PC */
-		public static string UserDevRootMac = "/UE4/Builds/";
+		/** The base path (on the Mac) to the your particular development directory, where files will be copied to from the PC */
+		public static string UserDevRootMacBase = "/UE4/Builds/";
+
+		/** The final path (on the Mac) to your particular development directory, where files will be copied to from the PC */
+		public static string UserDevRootMac = "/UE4/Builds";
 
 		/** Whether or not to connect to UnrealRemoteTool using RPCUtility */
 		[XmlConfig]
@@ -231,9 +234,13 @@ namespace UnrealBuildTool
 				{
 					// pass back the string
 					string HomeLocation = Results["CommandOutput"] as string;
-					UserDevRootMac = HomeLocation + UserDevRootMac;
+					UserDevRootMac = HomeLocation + UserDevRootMacBase;
 				}
 			}
+			else
+			{
+				UserDevRootMac = UserDevRootMacBase;
+			}
 		}
 
         // Do any one-time, global initialization for the tool chain
@@ -357,10 +364,10 @@ namespace UnrealBuildTool
 						}
 					}
 
-					if (!bFoundOverrideSSHPrivateKey)
+/*					if (!bFoundOverrideSSHPrivateKey)
 					{
 						throw new BuildException("An SSHKey was required, but one cannot be found. Can't continue...");
-					}
+					}*/
 
 					// resolve the rest of the strings
 					ResolvedRSyncExe = ResolveString(RSyncExe, true);
@@ -384,8 +391,8 @@ namespace UnrealBuildTool
 						ResolvedRSyncExe,
 						ResolvedRSyncUsername,
 						RemoteServerName,
-						Environment.GetFolderPath(Environment.SpecialFolder.Personal),
-						ConvertPathToCygwin(Environment.GetFolderPath(Environment.SpecialFolder.Personal)),
+						Environment.GetFolderPath(Environment.SpecialFolder.ApplicationData),
+						ConvertPathToCygwin(Environment.GetFolderPath(Environment.SpecialFolder.ApplicationData)),
 						Path.GetFullPath(BuildConfiguration.RelativeEnginePath));
 
 					KeyProcess.Start();
@@ -689,7 +696,7 @@ namespace UnrealBuildTool
 				// --exclude='*'  ??? why???
 				RsyncProcess.StartInfo.FileName = ResolvedRSyncExe;
 				RsyncProcess.StartInfo.Arguments = string.Format(
-					"-vzae \"{0}\" --rsync-path=\"mkdir -p {2} && rsync\" --chmod=ug=rwX,o=rxX --delete --files-from=\"{4}\" --include-from=\"{5}\" --include='*/' --exclude='*.o' --exclude='Timestamp' {1} {6}@{3}:'{2}'",
+					"-vzae \"{0}\" --rsync-path=\"mkdir -p {2} && rsync\" --chmod=ug=rwX,o=rxX --delete --files-from=\"{4}\" --include-from=\"{5}\" --include='*/' --exclude='*.o' --exclude='Timestamp' '{1}' {6}@{3}:'{2}'",
 					ResolvedRsyncAuthentication,
 					CygRootPath,
 					RemotePath,
@@ -697,6 +704,7 @@ namespace UnrealBuildTool
 					ConvertPathToCygwin(RSyncPathsFile),
 					ConvertPathToCygwin(IncludeFromFile),
 					RSyncUsername);
+				Console.WriteLine("Command: " + RsyncProcess.StartInfo.Arguments);
 
 				RsyncProcess.OutputDataReceived += new DataReceivedEventHandler(OutputReceivedForRsync);
 				RsyncProcess.ErrorDataReceived += new DataReceivedEventHandler(OutputReceivedForRsync);
@@ -716,6 +724,7 @@ namespace UnrealBuildTool
 		static public bool UploadFile(string LocalPath, string RemotePath)
 		{
 			string RemoteDir = Path.GetDirectoryName(RemotePath).Replace("\\", "/");
+			RemoteDir = RemoteDir.Replace(" ", "\\ ");
 			string RemoteFilename = Path.GetFileName(RemotePath);
 
 			// get the executable dir for SSH, so Rsync can call it easily
@@ -730,7 +739,7 @@ namespace UnrealBuildTool
 			// make simple rsync commandline to send a file
 			RsyncProcess.StartInfo.FileName = ResolvedRSyncExe;
 			RsyncProcess.StartInfo.Arguments = string.Format(
-				"-zae \"{0}\" --rsync-path=\"mkdir -p {1} && rsync\" \"{2}\" {3}@{4}:\"{1}/{5}\"",
+				"-zae \"{0}\" --rsync-path=\"mkdir -p {1} && rsync\" '{2}' {3}@{4}:'{1}/{5}'",
 				ResolvedRsyncAuthentication,
 				RemoteDir,
 				ConvertPathToCygwin(LocalPath),
@@ -750,6 +759,7 @@ namespace UnrealBuildTool
 		{
 			// get the executable dir for SSH, so Rsync can call it easily
 			string ExeDir = Path.GetDirectoryName(ResolvedSSHExe);
+			string RemoteDir = RemotePath.Replace(" ", "\\ ");
 
 			Process RsyncProcess = new Process();
 			if (ExeDir != "")
@@ -763,12 +773,12 @@ namespace UnrealBuildTool
 			// make simple rsync commandline to send a file
 			RsyncProcess.StartInfo.FileName = ResolvedRSyncExe;
 			RsyncProcess.StartInfo.Arguments = string.Format(
-				"-zae \"{0}\" {2}@{3}:\"{4}\" \"{1}\"",
+				"-zae \"{0}\" {2}@{3}:'{4}' \"{1}\"",
 				ResolvedRsyncAuthentication,
 				ConvertPathToCygwin(LocalPath),
 				RSyncUsername,
 				RemoteServerName,
-				RemotePath
+				RemoteDir
 				);
 
 			RsyncProcess.OutputDataReceived += new DataReceivedEventHandler(OutputReceivedForRsync);
@@ -785,10 +795,10 @@ namespace UnrealBuildTool
 			Console.WriteLine("Doing {0}", Command);
 
 			// make the commandline for other end
-			string RemoteCommandline = "cd " + WorkingDirectory;
+			string RemoteCommandline = "cd \"" + WorkingDirectory + "\"";
 			if (!string.IsNullOrWhiteSpace(RemoteOutputPath))
 			{
-				RemoteCommandline += " && mkdir -p " + Path.GetDirectoryName(RemoteOutputPath).Replace("\\", "/");
+				RemoteCommandline += " && mkdir -p \"" + Path.GetDirectoryName(RemoteOutputPath).Replace("\\", "/") + "\"";
 			}
 
 			// get the executable dir for SSH
diff --git a/Engine/Source/Runtime/Core/Private/Misc/MonitoredProcess.cpp b/Engine/Source/Runtime/Core/Private/Misc/MonitoredProcess.cpp
--- a/Engine/Source/Runtime/Core/Private/Misc/MonitoredProcess.cpp
+++ b/Engine/Source/Runtime/Core/Private/Misc/MonitoredProcess.cpp
@@ -6,7 +6,7 @@
 /* FMonitoredProcess structors
  *****************************************************************************/
 
-FMonitoredProcess::FMonitoredProcess( const FString& InURL, const FString& InParams, bool InHidden )
+FMonitoredProcess::FMonitoredProcess( const FString& InURL, const FString& InParams, bool InHidden, bool InCreatePipes )
 	: Canceling(false)
 	, EndTime(0)
 	, Hidden(InHidden)
@@ -18,6 +18,7 @@ FMonitoredProcess::FMonitoredProcess( const FString& InURL, const FString& InPar
 	, Thread(nullptr)
 	, URL(InURL)
 	, WritePipe(nullptr)
+	, bCreatePipes(InCreatePipes)
 { }
 
 
@@ -53,7 +54,7 @@ bool FMonitoredProcess::Launch()
 		return false;
 	}
 
-	if (!FPlatformProcess::CreatePipe(ReadPipe, WritePipe))
+	if (bCreatePipes && !FPlatformProcess::CreatePipe(ReadPipe, WritePipe))
 	{
 		return false;
 	}
diff --git a/Engine/Source/Runtime/Core/Public/Misc/MonitoredProcess.h b/Engine/Source/Runtime/Core/Public/Misc/MonitoredProcess.h
--- a/Engine/Source/Runtime/Core/Public/Misc/MonitoredProcess.h
+++ b/Engine/Source/Runtime/Core/Public/Misc/MonitoredProcess.h
@@ -33,7 +33,7 @@ public:
 	 * @param InParams The command line parameters.
 	 * @param InHidden Whether the window of the process should be hidden.
 	 */
-	FMonitoredProcess( const FString& InURL, const FString& InParams, bool InHidden );
+	FMonitoredProcess( const FString& InURL, const FString& InParams, bool InHidden, bool InCreatePipes = true );
 
 	/** Destructor. */
 	~FMonitoredProcess();
@@ -178,6 +178,9 @@ private:
 	// Holds the write pipe. */
 	void* WritePipe;
 
+	// Holds if we should create pipes
+	bool bCreatePipes;
+
 private:
 
 	// Holds a delegate that is executed when the process has been canceled. */
diff --git a/Engine/Source/Runtime/IOS/IOSRuntimeSettings/Private/IOSRuntimeSettings.cpp b/Engine/Source/Runtime/IOS/IOSRuntimeSettings/Private/IOSRuntimeSettings.cpp
--- a/Engine/Source/Runtime/IOS/IOSRuntimeSettings/Private/IOSRuntimeSettings.cpp
+++ b/Engine/Source/Runtime/IOS/IOSRuntimeSettings/Private/IOSRuntimeSettings.cpp
@@ -63,8 +63,12 @@ void UIOSRuntimeSettings::PostInitProperties()
 	// We can have a look for potential keys
 	if (!RemoteServerName.IsEmpty() && !RSyncUsername.IsEmpty())
 	{
+		SSHPrivateKeyLocation = TEXT("");
+
 		const FString DefaultKeyFilename = TEXT("RemoteToolChainPrivate.key");
 		const FString RelativeFilePathLocation = FPaths::Combine(TEXT("SSHKeys"), *RemoteServerName, *RSyncUsername, *DefaultKeyFilename);
+		TCHAR Path[4096];
+		FPlatformMisc::GetEnvironmentVariable(TEXT("APPDATA"), Path, ARRAY_COUNT(Path));
 
 		TArray<FString> PossibleKeyLocations;
 		PossibleKeyLocations.Add(FPaths::Combine(*FPaths::GameDir(), TEXT("Build"), TEXT("NotForLicensees"), *RelativeFilePathLocation));
@@ -73,6 +77,7 @@ void UIOSRuntimeSettings::PostInitProperties()
 		PossibleKeyLocations.Add(FPaths::Combine(*FPaths::EngineDir(), TEXT("Build"), TEXT("NotForLicensees"), *RelativeFilePathLocation));
 		PossibleKeyLocations.Add(FPaths::Combine(*FPaths::EngineDir(), TEXT("Build"), TEXT("NoRedist"), *RelativeFilePathLocation));
 		PossibleKeyLocations.Add(FPaths::Combine(*FPaths::EngineDir(), TEXT("Build"), *RelativeFilePathLocation));
+		PossibleKeyLocations.Add(FPaths::Combine(Path, TEXT("Unreal Engine"), TEXT("UnrealBuildTool"), *RelativeFilePathLocation));
 
 		// Find a potential path that we will use if the user hasn't overridden.
 		// For information purposes only
