disable IOS device network pings as we aren't using them right now and they seem to be causing occasional crashes
UE-11815
#ios

diff --git a/Engine/Source/Developer/iOS/IOSTargetPlatform/Private/IOSTargetPlatform.cpp b/Engine/Source/Developer/iOS/IOSTargetPlatform/Private/IOSTargetPlatform.cpp
--- a/Engine/Source/Developer/iOS/IOSTargetPlatform/Private/IOSTargetPlatform.cpp
+++ b/Engine/Source/Developer/iOS/IOSTargetPlatform/Private/IOSTargetPlatform.cpp
@@ -187,7 +187,8 @@ int32 FIOSTargetPlatform::CheckRequirements(const FString& ProjectPath, bool bPr
 
 void FIOSTargetPlatform::PingNetworkDevices()
 {
-	if (!MessageEndpoint.IsValid())
+	// disabled for now because we find IOS devices from the USB, this is a relic from ULD, but it may be needed in the future
+/*	if (!MessageEndpoint.IsValid())
 	{
 		MessageEndpoint = FMessageEndpoint::Builder("FIOSTargetPlatform")
 			.Handling<FIOSLaunchDaemonPong>(this, &FIOSTargetPlatform::HandlePongMessage);
@@ -210,7 +211,7 @@ void FIOSTargetPlatform::PingNetworkDevices()
 			DeviceIt.RemoveCurrent();
 			DeviceLostEvent.Broadcast(Device.ToSharedRef());
 		}
-	}
+	}*/
 }
 
 
