CIS fix

diff --git a/Engine/Source/Developer/iOS/IOSPlatformEditor/Private/IOSTargetSettingsCustomization.cpp b/Engine/Source/Developer/iOS/IOSPlatformEditor/Private/IOSTargetSettingsCustomization.cpp
--- a/Engine/Source/Developer/iOS/IOSPlatformEditor/Private/IOSTargetSettingsCustomization.cpp
+++ b/Engine/Source/Developer/iOS/IOSPlatformEditor/Private/IOSTargetSettingsCustomization.cpp
@@ -225,7 +225,7 @@ void FIOSTargetSettingsCustomization::UpdateStatus()
 	}
 }
 
-void FIOSTargetSettingsCustomization::UpdateStatus()
+void FIOSTargetSettingsCustomization::UpdateSSHStatus()
 {
 	// updated SSH key
 	const UIOSRuntimeSettings* Settings = GetDefault<UIOSRuntimeSettings>();
