Ensure correct assignment of emulated (ES2) uniform buffer indices for UB's generated within BuildResourceTableMapping(). Typically arises when only sampler related parameters are encountered.

Fixes UE-9804 and probably more.

diff --git a/Engine/Source/Developer/ShaderFormatOpenGL/Private/OpenGLShaderCompiler.cpp b/Engine/Source/Developer/ShaderFormatOpenGL/Private/OpenGLShaderCompiler.cpp
--- a/Engine/Source/Developer/ShaderFormatOpenGL/Private/OpenGLShaderCompiler.cpp
+++ b/Engine/Source/Developer/ShaderFormatOpenGL/Private/OpenGLShaderCompiler.cpp
@@ -878,6 +878,7 @@ static void BuildShaderOutput(
 		verify(Match(ShaderSource, '('));
 		uint16 BufferIndex = ParseNumber(ShaderSource);
 		check(BufferIndex == Header.Bindings.NumUniformBuffers);
+		UsedUniformBufferSlots[BufferIndex] = true;
 		verify(Match(ShaderSource, ')'));
 		ParameterMap.AddParameterAllocation(*BufferName, Header.Bindings.NumUniformBuffers++, 0, 0);
 
diff --git a/Engine/Source/Developer/iOS/MetalShaderFormat/Private/MetalShaderCompiler.cpp b/Engine/Source/Developer/iOS/MetalShaderFormat/Private/MetalShaderCompiler.cpp
--- a/Engine/Source/Developer/iOS/MetalShaderFormat/Private/MetalShaderCompiler.cpp
+++ b/Engine/Source/Developer/iOS/MetalShaderFormat/Private/MetalShaderCompiler.cpp
@@ -362,6 +362,7 @@ static void BuildMetalShaderOutput(
 		verify(Match(ShaderSource, '('));
 		uint16 BufferIndex = ParseNumber(ShaderSource);
 		check(BufferIndex == Header.Bindings.NumUniformBuffers);
+		UsedUniformBufferSlots[BufferIndex] = true;
 		verify(Match(ShaderSource, ')'));
 		ParameterMap.AddParameterAllocation(*BufferName, Header.Bindings.NumUniformBuffers++, 0, 0);
 
