UE4 - Metal - Change fmax/fmin to max/min when using int/uint types
- Change local var name to match type for easier debugging

diff --git a/Engine/Source/Developer/iOS/MetalShaderFormat/Private/MetalBackend.cpp b/Engine/Source/Developer/iOS/MetalShaderFormat/Private/MetalBackend.cpp
--- a/Engine/Source/Developer/iOS/MetalShaderFormat/Private/MetalBackend.cpp
+++ b/Engine/Source/Developer/iOS/MetalShaderFormat/Private/MetalBackend.cpp
@@ -451,9 +451,32 @@ protected:
 			const char *name = (const char *) hash_table_find(this->printable_names, var);
 			if (name == NULL)
 			{
-				bool is_global = (scope_depth == 0 && var->mode != ir_var_temporary);
-				const char *prefix = is_global ? "g" : "t";
-				int var_id = is_global ? global_id++ : temp_id++;
+				bool bIsGlobal = (scope_depth == 0 && var->mode != ir_var_temporary);
+				const char* prefix = "g";
+				if (!bIsGlobal)
+				{
+					if (var->type->is_matrix())
+					{
+						prefix = "m";
+					}
+					else if (var->type->is_vector())
+					{
+						prefix = "v";
+					}
+					else
+					{
+						switch (var->type->base_type)
+						{
+						case GLSL_TYPE_BOOL: prefix = "b"; break;
+						case GLSL_TYPE_UINT: prefix = "u"; break;
+						case GLSL_TYPE_INT: prefix = "i"; break;
+						case GLSL_TYPE_HALF: prefix = "h"; break;
+						case GLSL_TYPE_FLOAT: prefix = "f"; break;
+						default: prefix = "t"; break;
+						}
+					}
+				}
+				int var_id = bIsGlobal ? global_id++ : temp_id++;
 				name = ralloc_asprintf(mem_ctx, "%s%d", prefix, var_id);
 				hash_table_insert(this->printable_names, (void *)name, var);
 			}
@@ -1033,6 +1056,17 @@ protected:
 				ralloc_asprintf_append(buffer, MetalExpressionTable[op][i+1]);
 			}
 		}
+		else if (numOps == 2 && (op == ir_binop_max || op == ir_binop_min) && expr->type->is_integer())
+		{
+			// Convert fmax/fmin to max/min when dealing with integers
+			auto* OpString = MetalExpressionTable[op][0];
+			check(OpString[0] == 'f');
+			ralloc_asprintf_append(buffer, OpString + 1);
+			expr->operands[0]->accept(this);
+			ralloc_asprintf_append(buffer, MetalExpressionTable[op][1]);
+			expr->operands[1]->accept(this);
+			ralloc_asprintf_append(buffer, MetalExpressionTable[op][2]);
+		}
 		else if (numOps < 4)
 		{
 			ralloc_asprintf_append(buffer, MetalExpressionTable[op][0]);
