UE4 - Metal - Fix for depth textures
- Fix case of float3 Out : SV_Target1 instead of float4

diff --git a/Engine/Source/Developer/iOS/MetalShaderFormat/Private/MetalBackend.cpp b/Engine/Source/Developer/iOS/MetalShaderFormat/Private/MetalBackend.cpp
--- a/Engine/Source/Developer/iOS/MetalShaderFormat/Private/MetalBackend.cpp
+++ b/Engine/Source/Developer/iOS/MetalShaderFormat/Private/MetalBackend.cpp
@@ -562,7 +562,15 @@ protected:
 		else
 		{
 			check(t->HlslName);
-			if (bUsePacked && t->is_vector() && t->vector_elements < 4)
+			if (!strcmp(t->HlslName, "texture2d"))
+			{
+				ralloc_asprintf_append(buffer, "depth2d");
+			}
+			else if (!strcmp(t->HlslName, "texturecube"))
+			{
+				ralloc_asprintf_append(buffer, "depthcube");
+			}
+			else if (bUsePacked && t->is_vector() && t->vector_elements < 4)
 			{
 				ralloc_asprintf_append(buffer, "packed_%s", t->HlslName);
 			}
@@ -1085,14 +1093,14 @@ protected:
 		tex->sampler->accept(this);
 		if (tex->op == ir_tex || tex->op == ir_txl || tex->op == ir_txb)
 		{
-			ralloc_asprintf_append(buffer, ".sample(");
+			ralloc_asprintf_append(buffer, tex->shadow_comparitor ? ".sample_compare(" : ".sample(");
 			auto* Texture = tex->sampler->variable_referenced();
 			check(Texture);
 			auto* Entry = ParseState->FindPackedSamplerEntry(Texture->name);
 			ralloc_asprintf_append(buffer, "s%d, ", Entry->offset);
 			tex->coordinate->accept(this);
 
-			if (tex->op == ir_txl)
+			if (tex->op == ir_txl && !tex->shadow_comparitor)
 			{
 				ralloc_asprintf_append(buffer, ", level(");
 				tex->lod_info.lod->accept(this);
@@ -1105,6 +1113,12 @@ protected:
 				ralloc_asprintf_append(buffer, ")");
 			}
 
+			if (tex->shadow_comparitor)
+			{
+				ralloc_asprintf_append(buffer, ", ");
+				tex->shadow_comparitor->accept(this);
+			}
+
 			if (tex->offset)
 			{
 				ralloc_asprintf_append(buffer, ", ");
diff --git a/Engine/Source/Developer/iOS/MetalShaderFormat/Private/MetalUtils.cpp b/Engine/Source/Developer/iOS/MetalShaderFormat/Private/MetalUtils.cpp
--- a/Engine/Source/Developer/iOS/MetalShaderFormat/Private/MetalUtils.cpp
+++ b/Engine/Source/Developer/iOS/MetalShaderFormat/Private/MetalUtils.cpp
@@ -202,6 +202,30 @@ namespace MetalUtils
 		{"SV_Target5", glsl_type::half4_type, "FragColor5", ir_var_out, "[[ color(5) ]]"},
 		{"SV_Target6", glsl_type::half4_type, "FragColor6", ir_var_out, "[[ color(6) ]]"},
 		{"SV_Target7", glsl_type::half4_type, "FragColor7", ir_var_out, "[[ color(7) ]]"},
+		{"SV_Target0", glsl_type::half3_type, "FragColor0", ir_var_out, "[[ color(0) ]]"},
+		{"SV_Target1", glsl_type::half3_type, "FragColor1", ir_var_out, "[[ color(1) ]]"},
+		{"SV_Target2", glsl_type::half3_type, "FragColor2", ir_var_out, "[[ color(2) ]]"},
+		{"SV_Target3", glsl_type::half3_type, "FragColor3", ir_var_out, "[[ color(3) ]]"},
+		{"SV_Target4", glsl_type::half3_type, "FragColor4", ir_var_out, "[[ color(4) ]]"},
+		{"SV_Target5", glsl_type::half3_type, "FragColor5", ir_var_out, "[[ color(5) ]]"},
+		{"SV_Target6", glsl_type::half3_type, "FragColor6", ir_var_out, "[[ color(6) ]]"},
+		{"SV_Target7", glsl_type::half3_type, "FragColor7", ir_var_out, "[[ color(7) ]]"},
+		{"SV_Target0", glsl_type::half2_type, "FragColor0", ir_var_out, "[[ color(0) ]]"},
+		{"SV_Target1", glsl_type::half2_type, "FragColor1", ir_var_out, "[[ color(1) ]]"},
+		{"SV_Target2", glsl_type::half2_type, "FragColor2", ir_var_out, "[[ color(2) ]]"},
+		{"SV_Target3", glsl_type::half2_type, "FragColor3", ir_var_out, "[[ color(3) ]]"},
+		{"SV_Target4", glsl_type::half2_type, "FragColor4", ir_var_out, "[[ color(4) ]]"},
+		{"SV_Target5", glsl_type::half2_type, "FragColor5", ir_var_out, "[[ color(5) ]]"},
+		{"SV_Target6", glsl_type::half2_type, "FragColor6", ir_var_out, "[[ color(6) ]]"},
+		{"SV_Target7", glsl_type::half2_type, "FragColor7", ir_var_out, "[[ color(7) ]]"},
+		{"SV_Target0", glsl_type::half_type, "FragColor0", ir_var_out, "[[ color(0) ]]"},
+		{"SV_Target1", glsl_type::half_type, "FragColor1", ir_var_out, "[[ color(1) ]]"},
+		{"SV_Target2", glsl_type::half_type, "FragColor2", ir_var_out, "[[ color(2) ]]"},
+		{"SV_Target3", glsl_type::half_type, "FragColor3", ir_var_out, "[[ color(3) ]]"},
+		{"SV_Target4", glsl_type::half_type, "FragColor4", ir_var_out, "[[ color(4) ]]"},
+		{"SV_Target5", glsl_type::half_type, "FragColor5", ir_var_out, "[[ color(5) ]]"},
+		{"SV_Target6", glsl_type::half_type, "FragColor6", ir_var_out, "[[ color(6) ]]"},
+		{"SV_Target7", glsl_type::half_type, "FragColor7", ir_var_out, "[[ color(7) ]]"},
 		{NULL, NULL, NULL, ir_var_auto, nullptr}
 	};
 
@@ -460,7 +484,7 @@ namespace MetalUtils
 			FSystemValue* SystemValues = SystemValueTable[Frequency];
 			for (int i = 0; SystemValues[i].HlslSemantic != nullptr; ++i)
 			{
-				if (SystemValues[i].Mode == ir_var_out && FCStringAnsi::Stricmp(SystemValues[i].HlslSemantic, Semantic) == 0)
+				if (SystemValues[i].Mode == ir_var_out && FCStringAnsi::Stricmp(SystemValues[i].HlslSemantic, Semantic) == 0 && SystemValues[i].Type->vector_elements == Type->vector_elements)
 				{
 					Variable = new(ParseState) ir_variable(SystemValues[i].Type, SystemValues[i].MetalName, ir_var_out);
 					Variable->semantic = SystemValues[i].MetalSemantic;
@@ -1717,7 +1741,7 @@ void FMetalCodeBackend::PackInputsAndOutputs(exec_list* Instructions, _mesa_glsl
 					{
 						Mask = (1 << Variable->type->vector_elements) - 1;
 					}
-					auto* Assign = new(ParseState)ir_assignment(new(ParseState)ir_dereference_variable(Variable), DeRefMember, nullptr, Mask);
+					auto* Assign = new(ParseState)ir_assignment(new(ParseState) ir_dereference_variable(Variable), DeRefMember, nullptr, Mask);
 					PreCallInstructions.push_tail(Assign);
 					VarsToMoveToBody.push_back(Variable);
 				}
@@ -1740,15 +1764,16 @@ void FMetalCodeBackend::PackInputsAndOutputs(exec_list* Instructions, _mesa_glsl
 				if (VSOutVariables.find(Variable) != VSOutVariables.end())
 				{
 					VarsToMoveToBody.push_back(Variable);
-					ir_dereference* DeRefMember = new(ParseState)ir_dereference_record(VSOut, Variable->name);
-					auto* Assign = new(ParseState)ir_assignment(DeRefMember, new(ParseState)ir_dereference_variable(Variable));
+					ir_dereference* DeRefMember = new(ParseState) ir_dereference_record(VSOut, Variable->name);
+					auto* Assign = new(ParseState) ir_assignment(DeRefMember, new(ParseState)ir_dereference_variable(Variable));
 					PostCallInstructions.push_tail(Assign);
 				}
 				else if (PSOutVariables.find(Variable) != PSOutVariables.end())
 				{
 					VarsToMoveToBody.push_back(Variable);
-					ir_dereference* DeRefMember = new(ParseState)ir_dereference_record(PSOut, Variable->name);
-					auto* Assign = new(ParseState)ir_assignment(DeRefMember, new(ParseState)ir_dereference_variable(Variable));
+					ir_dereference* DeRefMember = new(ParseState) ir_dereference_record(PSOut, Variable->name);
+					ir_rvalue* DeRefVar = new(ParseState) ir_dereference_variable(Variable);
+					auto* Assign = new(ParseState) ir_assignment(DeRefMember, DeRefVar);
 					PostCallInstructions.push_tail(Assign);
 				}
 				else
