UE4 - Metal - fixed bug with previous check-in (depth samplers)
- Initial support for texture arrays

diff --git a/Engine/Source/Developer/iOS/MetalShaderFormat/Private/MetalBackend.cpp b/Engine/Source/Developer/iOS/MetalShaderFormat/Private/MetalBackend.cpp
--- a/Engine/Source/Developer/iOS/MetalShaderFormat/Private/MetalBackend.cpp
+++ b/Engine/Source/Developer/iOS/MetalShaderFormat/Private/MetalBackend.cpp
@@ -561,22 +561,44 @@ protected:
 		}
 		else
 		{
-			check(t->HlslName);
-			if (!strcmp(t->HlslName, "texture2d"))
+			if (t->base_type == GLSL_TYPE_SAMPLER)
 			{
-				ralloc_asprintf_append(buffer, "depth2d");
-			}
-			else if (!strcmp(t->HlslName, "texturecube"))
-			{
-				ralloc_asprintf_append(buffer, "depthcube");
-			}
-			else if (bUsePacked && t->is_vector() && t->vector_elements < 4)
-			{
-				ralloc_asprintf_append(buffer, "packed_%s", t->HlslName);
+				bool bDone = false;
+				if (t->sampler_dimensionality == GLSL_SAMPLER_DIM_2D && t->sampler_array)
+				{
+					ralloc_asprintf_append(buffer, t->sampler_shadow ? "depth2d_array" : "texture2d_array");
+					bDone = true;
+				}
+				else if (t->HlslName)
+				{
+					if (!strcmp(t->HlslName, "texture2d") && t->sampler_shadow)
+					{
+						ralloc_asprintf_append(buffer, "depth2d");
+						bDone = true;
+					}
+					else if (!strcmp(t->HlslName, "texturecube") && t->sampler_shadow)
+					{
+						ralloc_asprintf_append(buffer, "depthcube");
+						bDone = true;
+					}
+				}
+
+				if (!bDone)
+				{
+					ralloc_asprintf_append(buffer, "%s", t->HlslName ? t->HlslName : "UnsupportedSamplerType");
+				}
 			}
 			else
 			{
-				ralloc_asprintf_append(buffer, "%s", t->HlslName);
+				check(t->HlslName);
+				if (bUsePacked && t->is_vector() && t->vector_elements < 4)
+				{
+					ralloc_asprintf_append(buffer, "packed_%s", t->HlslName);
+				}
+				else
+				{
+					ralloc_asprintf_append(buffer, "%s", t->HlslName);
+				}
 			}
 		}
 	}
