UE4 - Metal - Fix texture arrays

diff --git a/Engine/Source/Developer/iOS/MetalShaderFormat/Private/MetalBackend.cpp b/Engine/Source/Developer/iOS/MetalShaderFormat/Private/MetalBackend.cpp
--- a/Engine/Source/Developer/iOS/MetalShaderFormat/Private/MetalBackend.cpp
+++ b/Engine/Source/Developer/iOS/MetalShaderFormat/Private/MetalBackend.cpp
@@ -1120,7 +1120,19 @@ protected:
 			check(Texture);
 			auto* Entry = ParseState->FindPackedSamplerEntry(Texture->name);
 			ralloc_asprintf_append(buffer, "s%d, ", Entry->offset);
-			tex->coordinate->accept(this);
+			if (tex->sampler->type->sampler_array)
+			{
+				// Need to split the coordinate
+				ralloc_asprintf_append(buffer, "(");
+				tex->coordinate->accept(this);
+				ralloc_asprintf_append(buffer, ").x%s, (uint)(", tex->sampler->type->sampler_dimensionality == GLSL_SAMPLER_DIM_2D ? "y" : "");
+				tex->coordinate->accept(this);
+				ralloc_asprintf_append(buffer, ").%s", tex->sampler->type->sampler_dimensionality == GLSL_SAMPLER_DIM_2D ? "z" : "y");
+			}
+			else
+			{
+				tex->coordinate->accept(this);
+			}
 
 			if (tex->op == ir_txl && !tex->shadow_comparitor)
 			{
