Adding deltacopy location to the configurable IOS runtime settings

Fixed some minor issues with the remote build settings UI
- Reordered properties slightly.
- Updated the Remote Server Name widget so it no longer has no size.
- Renamed some property display names to be a little clearer.

https://jira.ol.epicgames.net/browse/UE-14323

diff --git a/Engine/Source/Developer/iOS/IOSPlatformEditor/Private/IOSTargetSettingsCustomization.cpp b/Engine/Source/Developer/iOS/IOSPlatformEditor/Private/IOSTargetSettingsCustomization.cpp
--- a/Engine/Source/Developer/iOS/IOSPlatformEditor/Private/IOSTargetSettingsCustomization.cpp
+++ b/Engine/Source/Developer/iOS/IOSPlatformEditor/Private/IOSTargetSettingsCustomization.cpp
@@ -704,14 +704,6 @@ void FIOSTargetSettingsCustomization::BuildRemoteBuildingSection(IDetailLayoutBu
 	IDetailGroup& RemoteBuildingGroup = BuildCategory.AddGroup(*RemoteBuildingGroupName.ToString(), RemoteBuildingGroupName, false);
 
 
-	// Add Use RSync Property
-	TSharedRef<IPropertyHandle> UseRSyncPropertyHandle = DetailLayout.GetProperty(GET_MEMBER_NAME_CHECKED(UIOSRuntimeSettings, bUseRSync));
-	IDetailPropertyRow& UseRSyncPropertRow = RemoteBuildingGroup.AddPropertyRow(UseRSyncPropertyHandle);
-	UseRSyncPropertRow
-		.IsEnabled(!FRocketSupport::IsRocket())
-		.ToolTip(!FRocketSupport::IsRocket() ? LOCTEXT("UseRSyncToolTip", "Use RSync instead of RPCUtility") : FIOSTargetSettingsCustomizationConstants::DisabledTip);
-
-
 	// Remote Server Name Property
 	TSharedRef<IPropertyHandle> RemoteServerNamePropertyHandle = DetailLayout.GetProperty(GET_MEMBER_NAME_CHECKED(UIOSRuntimeSettings, RemoteServerName));
 	IDetailPropertyRow& RemoteServerNamePropertyRow = RemoteBuildingGroup.AddPropertyRow(RemoteServerNamePropertyHandle);
@@ -732,13 +724,11 @@ void FIOSTargetSettingsCustomization::BuildRemoteBuildingSection(IDetailLayoutBu
 			]
 		]
 		.ValueContent()
-		.MinDesiredWidth(0.0f)
-		.MaxDesiredWidth(0.0f)
+		.MinDesiredWidth(150.0f)
 		[
 			SNew(SHorizontalBox)
 			+ SHorizontalBox::Slot()
-			.FillWidth(1.0f)
-			.HAlign(HAlign_Fill)
+			.Padding(FMargin(0.0f, 8.0f))
 			[
 				SNew(SEditableTextBox)
 				.IsEnabled(this, &FIOSTargetSettingsCustomization::IsImportEnabled)
@@ -750,9 +740,21 @@ void FIOSTargetSettingsCustomization::BuildRemoteBuildingSection(IDetailLayoutBu
 				.ToolTipText(RemoteServerNamePropertyHandle->GetToolTipText())
 				.OnTextCommitted(this, &FIOSTargetSettingsCustomization::OnRemoteServerChanged, RemoteServerNamePropertyHandle)
 			]
+
 		];
 
-	
+
+
+	// Add Use RSync Property
+	TSharedRef<IPropertyHandle> UseRSyncPropertyHandle = DetailLayout.GetProperty(GET_MEMBER_NAME_CHECKED(UIOSRuntimeSettings, bUseRSync));
+	IDetailPropertyRow& UseRSyncPropertRow = RemoteBuildingGroup.AddPropertyRow(UseRSyncPropertyHandle);
+	UseRSyncPropertRow
+		.IsEnabled(!FRocketSupport::IsRocket())
+		.ToolTip(!FRocketSupport::IsRocket() ? LOCTEXT("UseRSyncToolTip", "Use RSync instead of RPCUtility") : FIOSTargetSettingsCustomizationConstants::DisabledTip);
+
+	TSharedRef<IPropertyHandle> DeltaCopyInstallDirPropertyHandle = DetailLayout.GetProperty(GET_MEMBER_NAME_CHECKED(UIOSRuntimeSettings, DeltaCopyInstallPath));
+	IDetailPropertyRow& DeltaCopyInstallDirPropertyRow = RemoteBuildingGroup.AddPropertyRow(DeltaCopyInstallDirPropertyHandle);
+
 	// Add RSync Username Property
 	TSharedRef<IPropertyHandle> RSyncUsernamePropertyHandle = DetailLayout.GetProperty(GET_MEMBER_NAME_CHECKED(UIOSRuntimeSettings, RSyncUsername));
 	IDetailPropertyRow& RSyncUsernamePropertyRow = RemoteBuildingGroup.AddPropertyRow(RSyncUsernamePropertyHandle);
@@ -772,14 +774,12 @@ void FIOSTargetSettingsCustomization::BuildRemoteBuildingSection(IDetailLayoutBu
 					.Font(DetailLayout.GetDetailFont())
 				]
 			]
-		.ValueContent()
-			.MinDesiredWidth(0.0f)
-			.MaxDesiredWidth(0.0f)
+			.ValueContent()
+			.MinDesiredWidth(150.0f)
 			[
 				SNew(SHorizontalBox)
 				+ SHorizontalBox::Slot()
-				.FillWidth(1.0f)
-				.HAlign(HAlign_Fill)
+				.Padding(FMargin(0.0f, 8.0f))
 				[
 					SNew(SEditableTextBox)
 					.IsEnabled(this, &FIOSTargetSettingsCustomization::IsImportEnabled)
@@ -1142,13 +1142,20 @@ FReply FIOSTargetSettingsCustomization::OnGenerateSSHKey()
 	}
 
 	FString CmdExe = FPaths::ConvertRelativePathToFull(FPaths::EngineDir() / TEXT("Build/BatchFiles/MakeAndInstallSSHKey.bat"));
-	TCHAR ProgramPath[4096];
-	FPlatformMisc::GetEnvironmentVariable(TEXT("PROGRAMFILES(X86)"), ProgramPath, ARRAY_COUNT(ProgramPath));
+	FString DeltaCopyPath = Settings.DeltaCopyInstallPath.Path;
+	if (DeltaCopyPath.IsEmpty())
+	{
+		// if the user hasn't specified a location for DeltaCopy, try and use the default install location
+		TCHAR ProgramPath[4096];
+		FPlatformMisc::GetEnvironmentVariable(TEXT("PROGRAMFILES(X86)"), ProgramPath, ARRAY_COUNT(ProgramPath));
+		DeltaCopyPath = FPaths::Combine(ProgramPath, TEXT("DeltaCopy"));
+
+	}
 	FString CygwinPath = TEXT("/cygdrive/") + FString(Path).Replace(TEXT(":"), TEXT("")).Replace(TEXT("\\"), TEXT("/"));
 	FString EnginePath = FPaths::EngineDir();
-	FString CommandLine = FString::Printf(TEXT("\"%s\\DeltaCopy\\ssh.exe\" \"%s\\DeltaCopy\\rsync.exe\" %s %s \"%s\" \"%s\" \"%s\""),
-		ProgramPath,
-		ProgramPath,
+	FString CommandLine = FString::Printf(TEXT("\"%s\\ssh.exe\" \"%s\\rsync.exe\" %s %s \"%s\" \"%s\" \"%s\""),
+		*DeltaCopyPath,
+		*DeltaCopyPath,
 		*(Settings.RSyncUsername),
 		*(Settings.RemoteServerName),
 		Path,
diff --git a/Engine/Source/Editor/DetailCustomizations/Private/DetailCustomizations.cpp b/Engine/Source/Editor/DetailCustomizations/Private/DetailCustomizations.cpp
--- a/Engine/Source/Editor/DetailCustomizations/Private/DetailCustomizations.cpp
+++ b/Engine/Source/Editor/DetailCustomizations/Private/DetailCustomizations.cpp
@@ -191,6 +191,7 @@ void FDetailCustomizationsModule::RegisterPropertyTypeCustomizations()
 	RegisterCustomPropertyTypeLayout("Margin", FOnGetPropertyTypeCustomizationInstance::CreateStatic(&FMarginStructCustomization::MakeInstance));
 	RegisterCustomPropertyTypeLayout("DirectoryPath", FOnGetPropertyTypeCustomizationInstance::CreateStatic(&FDirectoryPathStructCustomization::MakeInstance));
 	RegisterCustomPropertyTypeLayout("FilePath", FOnGetPropertyTypeCustomizationInstance::CreateStatic(&FFilePathStructCustomization::MakeInstance));
+	RegisterCustomPropertyTypeLayout("IOSBuildResourceDirectory", FOnGetPropertyTypeCustomizationInstance::CreateStatic(&FDirectoryPathStructCustomization::MakeInstance));
 	RegisterCustomPropertyTypeLayout("IOSBuildResourceFilePath", FOnGetPropertyTypeCustomizationInstance::CreateStatic(&FFilePathStructCustomization::MakeInstance));
 	RegisterCustomPropertyTypeLayout("InputAxisConfigEntry", FOnGetPropertyTypeCustomizationInstance::CreateStatic(&FInputAxisConfigCustomization::MakeInstance));
 	RegisterCustomPropertyTypeLayout("InputActionKeyMapping", FOnGetPropertyTypeCustomizationInstance::CreateStatic(&FInputActionMappingCustomization::MakeInstance));
diff --git a/Engine/Source/Programs/UnrealBuildTool/ToolChain/RemoteToolChain.cs b/Engine/Source/Programs/UnrealBuildTool/ToolChain/RemoteToolChain.cs
--- a/Engine/Source/Programs/UnrealBuildTool/ToolChain/RemoteToolChain.cs
+++ b/Engine/Source/Programs/UnrealBuildTool/ToolChain/RemoteToolChain.cs
@@ -42,6 +42,9 @@ namespace UnrealBuildTool
 		/** Whether or not to connect to UnrealRemoteTool using RPCUtility */
 		[XmlConfig]
 		public static bool bUseRPCUtil = true;
+
+		/** The user has specified a deltacopy install path */
+		private static string OverrideDeltaCopyInstallPath = null;
 		
 		/** Path to rsync executable and parameters for your rsync utility */
 		[XmlConfig]
@@ -206,6 +209,15 @@ namespace UnrealBuildTool
 			{
 				bUseRPCUtil = !bUseRSync;
 				Ini.GetString("/Script/IOSRuntimeSettings.IOSRuntimeSettings", "RSyncUsername", out RSyncUsername);
+				
+				if (Ini.GetString("/Script/IOSRuntimeSettings.IOSRuntimeSettings", "DeltaCopyInstallPath", out OverrideDeltaCopyInstallPath))
+				{
+					if (!string.IsNullOrEmpty(OverrideDeltaCopyInstallPath))
+					{
+						SSHExe = Path.Combine(OverrideDeltaCopyInstallPath, Path.GetFileName(SSHExe));
+						RSyncExe = Path.Combine(OverrideDeltaCopyInstallPath, Path.GetFileName(RSyncExe));
+					}
+				}
 
 				string ConfigKeyPath;
 				if (Ini.GetString("/Script/IOSRuntimeSettings.IOSRuntimeSettings", "SSHPrivateKeyOverridePath", out ConfigKeyPath))
diff --git a/Engine/Source/Runtime/IOS/IOSRuntimeSettings/Classes/IOSRuntimeSettings.h b/Engine/Source/Runtime/IOS/IOSRuntimeSettings/Classes/IOSRuntimeSettings.h
--- a/Engine/Source/Runtime/IOS/IOSRuntimeSettings/Classes/IOSRuntimeSettings.h
+++ b/Engine/Source/Runtime/IOS/IOSRuntimeSettings/Classes/IOSRuntimeSettings.h
@@ -83,6 +83,54 @@ struct TStructOpsTypeTraits<FIOSBuildResourceFilePath> : public TStructOpsTypeTr
 
 
 
+/**
+ *	IOS Build resource file struct, used to serialize Directorys to the configs for use in the build system,
+ */
+USTRUCT()
+struct FIOSBuildResourceDirectory
+{
+	GENERATED_USTRUCT_BODY()
+
+	/**
+	 * Custom export item used to serialize FIOSBuildResourceDirectory types as only a filename, no garland.
+	 */
+	bool ExportTextItem(FString& ValueStr, FIOSBuildResourceDirectory const& DefaultValue, UObject* Parent, int32 PortFlags, UObject* ExportRootScope) const
+	{
+		ValueStr += Path;
+		return true;
+	}
+
+	/**
+	 * Custom import item used to parse ini entries straight into the filename.
+	 */
+	bool ImportTextItem(const TCHAR*& Buffer, int32 PortFlags, UObject* Parent, FOutputDevice* ErrorText)
+	{
+		Path = Buffer;
+		return true;
+	}
+
+	/**
+	* The path to the file.
+	*/
+	UPROPERTY(EditAnywhere, Category = Directory)
+	FString Path;
+};
+
+/**
+*	Setup our resource Directory to make it easier to parse in UBT
+*/
+template<>
+struct TStructOpsTypeTraits<FIOSBuildResourceDirectory> : public TStructOpsTypeTraitsBase
+{
+	enum
+	{
+		WithExportTextItem = true,
+		WithImportTextItem = true,
+	};
+};
+
+
+
 /**
  * Implements the settings for the iOS target platform.
  */
@@ -137,19 +185,23 @@ public:
 	FString RemoteServerName;
 
 	// Enable the use of RSync for remote builds on a mac
-	UPROPERTY(GlobalConfig, EditAnywhere, Category = "Build", meta = (DisplayName = "Use RSync for IOS build"))
+	UPROPERTY(GlobalConfig, EditAnywhere, Category = "Build", meta = (DisplayName = "Use RSync for building IOS", ConfigHierarchyEditable))
 	bool bUseRSync;
 
 	// The mac users name which matches the SSH Private Key, for remote builds using RSync.
-	UPROPERTY(GlobalConfig, EditAnywhere, Category = "Build", meta = (EditCondition = "bUseRSync", DisplayName = "User for RSync enabled builds."))
+	UPROPERTY(GlobalConfig, EditAnywhere, Category = "Build", meta = (EditCondition = "bUseRSync", DisplayName = "Username on Remote Server.", ConfigHierarchyEditable))
 	FString RSyncUsername;
 
-	// The path of the ssh permissions key to be used when connecting to the remote server.
-	UPROPERTY(VisibleAnywhere, Category = "Build", meta = (DisplayName = "Existing SSH permissions file"))
+	// The install directory of DeltaCopy.
+	UPROPERTY(GlobalConfig, EditAnywhere, Category = "Build", meta = (EditCondition = "bUseRSync", ConfigHierarchyEditable))
+	FIOSBuildResourceDirectory DeltaCopyInstallPath;
+
+	// The existing location of an SSH Key found by UE4.
+	UPROPERTY(VisibleAnywhere, Category = "Build", meta = (DisplayName = "Found Existing SSH permissions file"))
 	FString SSHPrivateKeyLocation;
 
 	// The path of the ssh permissions key to be used when connecting to the remote server.
-	UPROPERTY(GlobalConfig, EditAnywhere, Category = "Build", meta = (EditCondition = "bUseRSync", DisplayName = "Override existing SSH permissions file"))
+	UPROPERTY(GlobalConfig, EditAnywhere, Category = "Build", meta = (EditCondition = "bUseRSync", DisplayName = "Override existing SSH permissions file", ConfigHierarchyEditable))
 	FIOSBuildResourceFilePath SSHPrivateKeyOverridePath;
 
 	// Supports default portrait orientation. Landscape will not be supported.
