UE4 - Metal - Support for gather

diff --git a/Engine/Source/Developer/iOS/MetalShaderFormat/Private/MetalBackend.cpp b/Engine/Source/Developer/iOS/MetalShaderFormat/Private/MetalBackend.cpp
--- a/Engine/Source/Developer/iOS/MetalShaderFormat/Private/MetalBackend.cpp
+++ b/Engine/Source/Developer/iOS/MetalShaderFormat/Private/MetalBackend.cpp
@@ -1177,6 +1177,40 @@ protected:
 				tex->coordinate->accept(this);
 			}
 		}
+		else if (tex->op == ir_txg)
+		{
+			//Tv gather(sampler s, float2 coord, int2 offset = int2(0)) const
+			//Tv gather_compare(sampler s, float2 coord, float compare_value, int2 offset = int2(0)) const
+			if (tex->shadow_comparitor)
+			{
+				ralloc_asprintf_append(buffer, ".gather_compare(");
+			}
+			else
+			{
+				ralloc_asprintf_append(buffer, ".gather(");
+			}
+			// Sampler
+			auto* Texture = tex->sampler->variable_referenced();
+			check(Texture);
+			auto* Entry = ParseState->FindPackedSamplerEntry(Texture->name);
+			int32 SamplerStateIndex = Buffers.GetUniqueSamplerStateIndex(tex->SamplerStateName, false);
+			check(SamplerStateIndex != INDEX_NONE);
+			ralloc_asprintf_append(buffer, "s%d, ", SamplerStateIndex);
+			// Coord
+			tex->coordinate->accept(this);
+
+			if (tex->shadow_comparitor)
+			{
+				tex->shadow_comparitor->accept(this);
+				ralloc_asprintf_append(buffer, ", ");
+			}
+
+			if (tex->offset)
+			{
+				ralloc_asprintf_append(buffer, ", ");
+				tex->offset->accept(this);
+			}
+		}
 		else
 		{
 			ralloc_asprintf_append(buffer, "UNKNOWN TEXOP %d!", tex->op);
