UE4 - Metal - Support for GetDimensions()

diff --git a/Engine/Source/Developer/iOS/MetalShaderFormat/Private/MetalBackend.cpp b/Engine/Source/Developer/iOS/MetalShaderFormat/Private/MetalBackend.cpp
--- a/Engine/Source/Developer/iOS/MetalShaderFormat/Private/MetalBackend.cpp
+++ b/Engine/Source/Developer/iOS/MetalShaderFormat/Private/MetalBackend.cpp
@@ -1113,8 +1113,17 @@ protected:
 	{
 		check(scope_depth > 0);
 		bool bNeedsClosingParenthesis = true;
+		if (tex->op == ir_txs)
+		{
+			ralloc_asprintf_append(buffer, "uint2(");
+		}
+
 		tex->sampler->accept(this);
-		if (tex->op == ir_tex || tex->op == ir_txl || tex->op == ir_txb)
+		switch (tex->op)
+		{
+		case ir_tex:
+		case ir_txl:
+		case ir_txb:
 		{
 			ralloc_asprintf_append(buffer, tex->shadow_comparitor ? ".sample_compare(" : ".sample(");
 			auto* Texture = tex->sampler->variable_referenced();
@@ -1162,7 +1171,9 @@ protected:
 				tex->offset->accept(this);
 			}
 		}
-		else if (tex->op == ir_txf)
+			break;
+
+		case ir_txf:
 		{
 			if (tex->sampler->type && tex->sampler->type->is_sampler() && tex->sampler->type->sampler_buffer)
 			{
@@ -1177,7 +1188,9 @@ protected:
 				tex->coordinate->accept(this);
 			}
 		}
-		else if (tex->op == ir_txg)
+			break;
+
+		case ir_txg:
 		{
 			//Tv gather(sampler s, float2 coord, int2 offset = int2(0)) const
 			//Tv gather_compare(sampler s, float2 coord, float compare_value, int2 offset = int2(0)) const
@@ -1211,10 +1224,51 @@ protected:
 				tex->offset->accept(this);
 			}
 		}
-		else
+			break;
+
+		case ir_txs:
+		{
+			// Convert from:
+			//	HLSL
+			//		uint w, h;
+			//		T.GetDimensions({lod, }w, h);
+			// GLSL
+			//		ivec2 Temp;
+			//		Temp = textureSize(T{, lod});
+			// Metal
+			//		uint2 Temp = uint2(T.get_width({lod}), T.get_height({lod}));
+			ralloc_asprintf_append(buffer, ".get_width(");
+			if (tex->lod_info.lod)
+			{
+				tex->lod_info.lod->accept(this);
+			}
+			ralloc_asprintf_append(buffer, "), ");
+			tex->sampler->accept(this);
+			ralloc_asprintf_append(buffer, ".get_height(");
+			if (tex->lod_info.lod)
+			{
+				tex->lod_info.lod->accept(this);
+			}
+			ralloc_asprintf_append(buffer, ")");
+		}
+			break;
+
+		case ir_txm:
 		{
+			// Convert from:
+			//	HLSL
+			//		uint w, h, d;
+			//		T.GetDimensions({lod, }w, h, d);
+			// Metal
+			//		uint2 Temp = T.get_num_mip_levels();
+			ralloc_asprintf_append(buffer, ".get_num_mip_levels()");
+		}
+			break;
+
+		default:
 			ralloc_asprintf_append(buffer, "UNKNOWN TEXOP %d!", tex->op);
 			check(0);
+			break;
 		}
 
 		if (bNeedsClosingParenthesis)
