UE4 - Metal - fix intrinsics

diff --git a/Engine/Source/Developer/iOS/MetalShaderFormat/Private/MetalBackend.cpp b/Engine/Source/Developer/iOS/MetalShaderFormat/Private/MetalBackend.cpp
--- a/Engine/Source/Developer/iOS/MetalShaderFormat/Private/MetalBackend.cpp
+++ b/Engine/Source/Developer/iOS/MetalShaderFormat/Private/MetalBackend.cpp
@@ -2960,11 +2960,11 @@ void FMetalLanguageSpec::SetupLanguageIntrinsics(_mesa_glsl_parse_state* State,
 
 	// Memory sync/barriers
 	{
-		make_intrinsic_genType(ir, State, GROUP_MEMORY_BARRIER, ir_invalid_opcode, 0, 0, 0, 0);
-		make_intrinsic_genType(ir, State, GROUP_MEMORY_BARRIER_WITH_GROUP_SYNC, ir_invalid_opcode, 0, 0, 0, 0);
-		make_intrinsic_genType(ir, State, DEVICE_MEMORY_BARRIER, ir_invalid_opcode, 0, 0, 0, 0);
-		make_intrinsic_genType(ir, State, DEVICE_MEMORY_BARRIER_WITH_GROUP_SYNC, ir_invalid_opcode, 0, 0, 0, 0);
-		make_intrinsic_genType(ir, State, ALL_MEMORY_BARRIER, ir_invalid_opcode, 0, 0, 0, 0);
-		make_intrinsic_genType(ir, State, ALL_MEMORY_BARRIER_WITH_GROUP_SYNC, ir_invalid_opcode, 0, 0, 0, 0);
+		make_intrinsic_genType(ir, State, GROUP_MEMORY_BARRIER, ir_invalid_opcode, IR_INTRINSIC_RETURNS_VOID, 0, 0, 0);
+		make_intrinsic_genType(ir, State, GROUP_MEMORY_BARRIER_WITH_GROUP_SYNC, ir_invalid_opcode, IR_INTRINSIC_RETURNS_VOID, 0, 0, 0);
+		make_intrinsic_genType(ir, State, DEVICE_MEMORY_BARRIER, ir_invalid_opcode, IR_INTRINSIC_RETURNS_VOID, 0, 0, 0);
+		make_intrinsic_genType(ir, State, DEVICE_MEMORY_BARRIER_WITH_GROUP_SYNC, ir_invalid_opcode, IR_INTRINSIC_RETURNS_VOID, 0, 0, 0);
+		make_intrinsic_genType(ir, State, ALL_MEMORY_BARRIER, ir_invalid_opcode, IR_INTRINSIC_RETURNS_VOID, 0, 0, 0);
+		make_intrinsic_genType(ir, State, ALL_MEMORY_BARRIER_WITH_GROUP_SYNC, ir_invalid_opcode, IR_INTRINSIC_RETURNS_VOID, 0, 0, 0);
 	}
 }
