SSH on by default
fix for No Code project from windows message
UE-14864
#ios

diff --git a/Engine/Config/BaseEngine.ini b/Engine/Config/BaseEngine.ini
--- a/Engine/Config/BaseEngine.ini
+++ b/Engine/Config/BaseEngine.ini
@@ -2128,7 +2128,7 @@ bShipForArmV7=True
 bShipForArm64=True
 bShipForArmV7S=False
 bGeneratedSYMFile=False
-bUseRSync=False
+bUseRSync=True
 BundleDisplayName=[PROJECT_NAME]
 BundleName=[PROJECT_NAME]
 BundleIdentifier=com.YourCompany.[PROJECT_NAME]
diff --git a/Engine/Source/Developer/iOS/IOSPlatformEditor/Private/IOSTargetSettingsCustomization.cpp b/Engine/Source/Developer/iOS/IOSPlatformEditor/Private/IOSTargetSettingsCustomization.cpp
--- a/Engine/Source/Developer/iOS/IOSPlatformEditor/Private/IOSTargetSettingsCustomization.cpp
+++ b/Engine/Source/Developer/iOS/IOSPlatformEditor/Private/IOSTargetSettingsCustomization.cpp
@@ -708,8 +708,7 @@ void FIOSTargetSettingsCustomization::BuildRemoteBuildingSection(IDetailLayoutBu
 	TSharedRef<IPropertyHandle> RemoteServerNamePropertyHandle = DetailLayout.GetProperty(GET_MEMBER_NAME_CHECKED(UIOSRuntimeSettings, RemoteServerName));
 	IDetailPropertyRow& RemoteServerNamePropertyRow = RemoteBuildingGroup.AddPropertyRow(RemoteServerNamePropertyHandle);
 	RemoteServerNamePropertyRow
-		.IsEnabled(!FRocketSupport::IsRocket())
-		.ToolTip(!FRocketSupport::IsRocket() ? LOCTEXT("RemoteServerNameToolTip", "The name or ip address of the remote mac which will be used to build IOS") : FIOSTargetSettingsCustomizationConstants::DisabledTip)
+		.ToolTip(LOCTEXT("RemoteServerNameToolTip", "The name or ip address of the remote mac which will be used to build IOS"))
 		.CustomWidget()
 		.NameContent()
 		[
@@ -747,10 +746,8 @@ void FIOSTargetSettingsCustomization::BuildRemoteBuildingSection(IDetailLayoutBu
 
 	// Add Use RSync Property
 	TSharedRef<IPropertyHandle> UseRSyncPropertyHandle = DetailLayout.GetProperty(GET_MEMBER_NAME_CHECKED(UIOSRuntimeSettings, bUseRSync));
-	IDetailPropertyRow& UseRSyncPropertRow = RemoteBuildingGroup.AddPropertyRow(UseRSyncPropertyHandle);
-	UseRSyncPropertRow
-		.IsEnabled(!FRocketSupport::IsRocket())
-		.ToolTip(!FRocketSupport::IsRocket() ? LOCTEXT("UseRSyncToolTip", "Use RSync instead of RPCUtility") : FIOSTargetSettingsCustomizationConstants::DisabledTip);
+	BuildCategory.AddProperty(UseRSyncPropertyHandle)
+		.Visibility(EVisibility::Hidden);
 
 	TSharedRef<IPropertyHandle> DeltaCopyInstallDirPropertyHandle = DetailLayout.GetProperty(GET_MEMBER_NAME_CHECKED(UIOSRuntimeSettings, DeltaCopyInstallPath));
 	IDetailPropertyRow& DeltaCopyInstallDirPropertyRow = RemoteBuildingGroup.AddPropertyRow(DeltaCopyInstallDirPropertyHandle);
@@ -759,8 +756,7 @@ void FIOSTargetSettingsCustomization::BuildRemoteBuildingSection(IDetailLayoutBu
 	TSharedRef<IPropertyHandle> RSyncUsernamePropertyHandle = DetailLayout.GetProperty(GET_MEMBER_NAME_CHECKED(UIOSRuntimeSettings, RSyncUsername));
 	IDetailPropertyRow& RSyncUsernamePropertyRow = RemoteBuildingGroup.AddPropertyRow(RSyncUsernamePropertyHandle);
 	RSyncUsernamePropertyRow
-		.IsEnabled(!FRocketSupport::IsRocket())
-		.ToolTip(!FRocketSupport::IsRocket() ? LOCTEXT("RSyncUsernameToolTip", "The username of the mac user that matches the specified SSH Key.") : FIOSTargetSettingsCustomizationConstants::DisabledTip)
+		.ToolTip(LOCTEXT("RSyncUsernameToolTip", "The username of the mac user that matches the specified SSH Key."))
 		.CustomWidget()
 			.NameContent()
 			[
@@ -798,16 +794,14 @@ void FIOSTargetSettingsCustomization::BuildRemoteBuildingSection(IDetailLayoutBu
 	TSharedRef<IPropertyHandle> SSHPrivateKeyLocationPropertyHandle = DetailLayout.GetProperty(GET_MEMBER_NAME_CHECKED(UIOSRuntimeSettings, SSHPrivateKeyLocation));
 	IDetailPropertyRow& SSHPrivateKeyLocationPropertyRow = RemoteBuildingGroup.AddPropertyRow(SSHPrivateKeyLocationPropertyHandle);
 	SSHPrivateKeyLocationPropertyRow
-		.IsEnabled(!FRocketSupport::IsRocket())
-		.ToolTip(!FRocketSupport::IsRocket() ? LOCTEXT("SSHPrivateKeyLocationToolTip", "The existing location of an SSH Key found by UE4.") : FIOSTargetSettingsCustomizationConstants::DisabledTip);
+		.ToolTip(LOCTEXT("SSHPrivateKeyLocationToolTip", "The existing location of an SSH Key found by UE4."));
 
 
 	// Add SSH override path
 	TSharedRef<IPropertyHandle> SSHPrivateKeyOverridePathPropertyHandle = DetailLayout.GetProperty(GET_MEMBER_NAME_CHECKED(UIOSRuntimeSettings, SSHPrivateKeyOverridePath));
 	IDetailPropertyRow& SSHPrivateKeyOverridePathPropertyRow = RemoteBuildingGroup.AddPropertyRow(SSHPrivateKeyOverridePathPropertyHandle);
 	SSHPrivateKeyOverridePathPropertyRow
-		.IsEnabled(!FRocketSupport::IsRocket())
-		.ToolTip(!FRocketSupport::IsRocket() ? LOCTEXT("SSHPrivateKeyOverridePathToolTip", "Override the existing SSH Private Key with one from a specified location.") : FIOSTargetSettingsCustomizationConstants::DisabledTip);
+		.ToolTip(LOCTEXT("SSHPrivateKeyOverridePathToolTip", "Override the existing SSH Private Key with one from a specified location."));
 
 	// Add a generate key button
 	RemoteBuildingGroup.AddWidgetRow()
diff --git a/Engine/Source/Developer/iOS/IOSTargetPlatform/Private/IOSTargetPlatform.cpp b/Engine/Source/Developer/iOS/IOSTargetPlatform/Private/IOSTargetPlatform.cpp
--- a/Engine/Source/Developer/iOS/IOSTargetPlatform/Private/IOSTargetPlatform.cpp
+++ b/Engine/Source/Developer/iOS/IOSTargetPlatform/Private/IOSTargetPlatform.cpp
@@ -119,17 +119,6 @@ int32 FIOSTargetPlatform::CheckRequirements(const FString& ProjectPath, bool bPr
 #if PLATFORM_MAC
 	OutTutorialPath = FString("/Engine/Tutorial/Installation/InstallingXCodeTutorial.InstallingXCodeTutorial");
     // shell to certtool
-#else
-	if (bProjectHasCode && FRocketSupport::IsRocket())
-	{
-		OutTutorialPath = FString("/Engine/Tutorial/Mobile/iOSonPCRestrictions.iOSonPCRestrictions");
-		bReadyToBuild |= ETargetPlatformReadyStatus::CodeUnsupported;
-	}
-	if (FRocketSupport::IsRocket() && IProjectManager::Get().IsNonDefaultPluginEnabled())
-	{
-		OutTutorialPath = FString("/Engine/Tutorial/Mobile/iOSonPCValidPlugins.iOSonPCValidPlugins");
-		bReadyToBuild |= ETargetPlatformReadyStatus::PluginsUnsupported;
-	}
 #endif
 
 	// shell to IPP and get the status of the provision and cert
diff --git a/Engine/Source/Runtime/IOS/IOSRuntimeSettings/Private/IOSRuntimeSettings.cpp b/Engine/Source/Runtime/IOS/IOSRuntimeSettings/Private/IOSRuntimeSettings.cpp
--- a/Engine/Source/Runtime/IOS/IOSRuntimeSettings/Private/IOSRuntimeSettings.cpp
+++ b/Engine/Source/Runtime/IOS/IOSRuntimeSettings/Private/IOSRuntimeSettings.cpp
@@ -23,7 +23,7 @@ UIOSRuntimeSettings::UIOSRuntimeSettings(const FObjectInitializer& ObjectInitial
 	bShipForArmV7 = true;
 	bShipForArm64 = true;
 	bShipForArmV7S = false;
-	bUseRSync = false;
+	bUseRSync = true;
 	AdditionalPlistData = TEXT("");
 }
 
