disable SSH in a binary build, by not allowing compiling for iOS
UE-15080
#ios

diff --git a/Engine/Source/Developer/iOS/IOSTargetPlatform/Private/IOSTargetPlatform.cpp b/Engine/Source/Developer/iOS/IOSTargetPlatform/Private/IOSTargetPlatform.cpp
--- a/Engine/Source/Developer/iOS/IOSTargetPlatform/Private/IOSTargetPlatform.cpp
+++ b/Engine/Source/Developer/iOS/IOSTargetPlatform/Private/IOSTargetPlatform.cpp
@@ -119,6 +119,17 @@ int32 FIOSTargetPlatform::CheckRequirements(const FString& ProjectPath, bool bPr
 #if PLATFORM_MAC
 	OutTutorialPath = FString("/Engine/Tutorial/Installation/InstallingXCodeTutorial.InstallingXCodeTutorial");
     // shell to certtool
+#else
+	if (bProjectHasCode && FRocketSupport::IsRocket())
+	{
+		OutTutorialPath = FString("/Engine/Tutorial/Mobile/iOSonPCRestrictions.iOSonPCRestrictions");
+		bReadyToBuild |= ETargetPlatformReadyStatus::CodeUnsupported;
+	}
+	if (FRocketSupport::IsRocket() && IProjectManager::Get().IsNonDefaultPluginEnabled())
+	{
+		OutTutorialPath = FString("/Engine/Tutorial/Mobile/iOSonPCValidPlugins.iOSonPCValidPlugins");
+		bReadyToBuild |= ETargetPlatformReadyStatus::PluginsUnsupported;
+	}
 #endif
 
 	// shell to IPP and get the status of the provision and cert
