- Improved power of 2 size checking for IOS textures, supports cubemaps, etc.

diff --git a/Engine/Source/Developer/iOS/IOSTargetPlatform/Private/IOSTargetPlatform.cpp b/Engine/Source/Developer/iOS/IOSTargetPlatform/Private/IOSTargetPlatform.cpp
--- a/Engine/Source/Developer/iOS/IOSTargetPlatform/Private/IOSTargetPlatform.cpp
+++ b/Engine/Source/Developer/iOS/IOSTargetPlatform/Private/IOSTargetPlatform.cpp
@@ -434,9 +434,8 @@ void FIOSTargetPlatform::GetTextureFormats( const UTexture* Texture, TArray<FNam
 			}
 			if (bIncludePVRTC)
 			{
-#if WITH_EDITORONLY_DATA
 				// handle non-power of 2 textures
-				if (!FMath::IsPowerOfTwo(Texture->Source.GetSizeX()) || !FMath::IsPowerOfTwo(Texture->Source.GetSizeY()))
+				if (!Texture->Source.IsPowerOfTwo())
 				{
 					// option 1: Uncompress, but users will get very large textures unknowningly
 					// OutFormats.AddUnique(NameBGRA8);
@@ -444,7 +443,6 @@ void FIOSTargetPlatform::GetTextureFormats( const UTexture* Texture, TArray<FNam
 					OutFormats.AddUnique(NamePOTERROR);
 				}
 				else
-#endif
 				{
 					OutFormats.AddUnique(FormatRemap[RemapIndex + 1]);
 				}
