UE4 - Metal integration from Zebra
- Shader changes/fix-ups, support for MRT & SM5 on the shader code gen backend

diff --git a/Engine/Shaders/TonemapCommon.usf b/Engine/Shaders/TonemapCommon.usf
--- a/Engine/Shaders/TonemapCommon.usf
+++ b/Engine/Shaders/TonemapCommon.usf
@@ -5,7 +5,7 @@
 =============================================================================*/
 
 // #todo-rco: Fix in the cross-compiler
-#if COMPILER_GLSL || COMPILER_GLSL_ES2 || METAL_PROFILE || VULKAN_PROFILE
+#if COMPILER_GLSL || COMPILER_GLSL_ES2 || METAL_PROFILE || VULKAN_PROFILE || METAL_MRT_PROFILE || METAL_SM5_PROFILE
 #define log10(x)  log(x) / log(10.0)
 #endif
 
diff --git a/Engine/Source/Developer/iOS/MetalShaderFormat/Private/MetalBackend.cpp b/Engine/Source/Developer/iOS/MetalShaderFormat/Private/MetalBackend.cpp
--- a/Engine/Source/Developer/iOS/MetalShaderFormat/Private/MetalBackend.cpp
+++ b/Engine/Source/Developer/iOS/MetalShaderFormat/Private/MetalBackend.cpp
@@ -1,5 +1,5 @@
 // Copyright 1998-2015 Epic Games, Inc. All Rights Reserved.
-// .
+// ..
 
 #include "MetalShaderFormat.h"
 #include "Core.h"
@@ -334,6 +334,7 @@ static void DumpSortedRanges(TDMARangeList& SortedRanges)
  */
 class FGenerateMetalVisitor : public ir_visitor
 {
+	FMetalCodeBackend* Backend;
 	_mesa_glsl_parse_state* ParseState;
 
 	/** Track which multi-dimensional arrays are used. */
@@ -373,7 +374,6 @@ protected:
 	/** Uniform block variables required by the code. */
 	hash_table *used_uniform_blocks;
 	/** Multi-dimensional arrays required by the code. */
-	exec_list used_md_arrays;
 
 	// Code generation flags
 	_mesa_glsl_parser_targets Frequency;
@@ -388,6 +388,8 @@ protected:
 	int indentation;
 	/** Scope depth. */
 	int scope_depth;
+	// Expression Depth
+	int ExpressionDepth;
 	/** The number of temporary variables declared in the current scope. */
 	int temp_id;
 	/** The number of global variables declared. */
@@ -416,31 +418,6 @@ protected:
 	// Do we need to add #include <compute_shaders>
 	bool bNeedsComputeInclude;
 
-	/**
-	 * Return true if the type is a multi-dimensional array. Also, track the
-	 * array.
-	 */
-	bool is_md_array(const glsl_type* type)
-	{
-		if (type->base_type == GLSL_TYPE_ARRAY &&
-			type->fields.array->base_type == GLSL_TYPE_ARRAY)
-		{
-			foreach_iter(exec_list_iterator, iter, used_md_arrays) 
-			{
-				md_array_entry* entry = (md_array_entry*)iter.get();
-				if (entry->type == type)
-				{
-					return true;
-				}
-			}
-			md_array_entry* entry = new(mem_ctx) md_array_entry();
-			entry->type = type;
-			used_md_arrays.push_tail(entry);
-			return true;
-		}
-		return false;
-	}
-
 	/**
 	 * Fetch/generate a unique name for ir_variable.
 	 *
@@ -502,23 +479,6 @@ protected:
 		}
 	}
 
-	/**
-	 * Print out the internal name for a multi-dimensional array.
-	 */
-	void print_md_array_type(const glsl_type *t)
-	{
-		if (t->base_type == GLSL_TYPE_ARRAY)
-		{
-			ralloc_asprintf_append(buffer, "_mdarr_");
-			do 
-			{
-				ralloc_asprintf_append(buffer, "%d_", t->length);
-				t = t->fields.array;
-			} while (t->base_type == GLSL_TYPE_ARRAY);
-			print_base_type(t);
-		}
-	}
-
 	/**
 	 * Print the base type, e.g. vec3.
 	 */
@@ -573,6 +533,16 @@ protected:
 					ralloc_asprintf_append(buffer, t->sampler_shadow ? "depth2d_array" : "texture2d_array");
 					bDone = true;
 				}
+				else if (t->sampler_dimensionality == GLSL_SAMPLER_DIM_CUBE && t->sampler_array)
+				{
+					ralloc_asprintf_append(buffer, "texturecube_array");
+					bDone = true;
+				}
+				else if (t->sampler_dimensionality == GLSL_SAMPLER_DIM_2D && t->sampler_ms)
+				{
+					ralloc_asprintf_append(buffer, t->sampler_shadow ? "depth2d_ms" : "texture2d_ms");
+					bDone = true;
+				}
 				else if (t->HlslName)
 				{
 					if (!strcmp(t->HlslName, "texture2d") && t->sampler_shadow)
@@ -612,14 +582,7 @@ protected:
 	 */
 	void print_type_pre(const glsl_type *t)
 	{
-		if (is_md_array(t))
-		{
-			print_md_array_type(t);
-		}
-		else
-		{
-			print_base_type(t);
-		}
+		print_base_type(t);
 	}
 
 	/**
@@ -627,9 +590,14 @@ protected:
 	 */
 	void print_type_post(const glsl_type *t)
 	{
-		if (t->base_type == GLSL_TYPE_ARRAY && !is_md_array(t))
+		if (t->base_type == GLSL_TYPE_ARRAY)
 		{
-			ralloc_asprintf_append(buffer, "[%u]", t->length);
+			do 
+			{
+				ralloc_asprintf_append(buffer, "[%u]", t->length);
+				t = t->element_type();
+			}
+			while (t && t->base_type == GLSL_TYPE_ARRAY);
 		}
 		else if (t->base_type == GLSL_TYPE_INPUTPATCH || t->base_type == GLSL_TYPE_OUTPUTPATCH)
 		{
@@ -706,7 +674,7 @@ protected:
 				check(!PtrType->is_array() && PtrType->inner_type);
 
 				// Buffer
-				int BufferIndex = Buffers.GetIndex(var);
+				int BufferIndex = Buffers.GetIndex(var, Backend->bIsDesktop);
 				check(BufferIndex >= 0);
 				if (var->type->sampler_buffer)
 				{
@@ -714,6 +682,10 @@ protected:
 						buffer,
 						"device "
 						);
+					if (Buffers.AtomicVariables.find(var) != Buffers.AtomicVariables.end())
+					{
+						ralloc_asprintf_append(buffer, "atomic_");
+					}
 					print_type_pre(PtrType->inner_type);
 					ralloc_asprintf_append(buffer, " *%s", unique_name(var));
 					print_type_post(PtrType->inner_type);
@@ -738,11 +710,12 @@ protected:
 					while (*Found);
 					*TempPtr = 0;
 					ralloc_asprintf_append(buffer, "texture%s<", Temp);
-					print_type_pre(PtrType->inner_type);
+					// UAVs require type per channel, not including # of channels
+					print_type_pre(PtrType->inner_type->get_scalar_type());
 					ralloc_asprintf_append(buffer, ", access::write> %s", unique_name(var));
 					ralloc_asprintf_append(
 						buffer,
-						" [[ buffer(%d) ]]", BufferIndex
+						" [[ texture(%d) ]]", BufferIndex
 						);
 				}
 			}
@@ -765,11 +738,11 @@ protected:
 						if (var->type->sampler_buffer)
 						{
 							// Buffer
-							int BufferIndex = Buffers.GetIndex(var);
+							int BufferIndex = Buffers.GetIndex(var, Backend->bIsDesktop);
 							check(BufferIndex >= 0);
 							ralloc_asprintf_append(
 								buffer,
-								"constant "
+								"const device "
 								);
 							print_type_pre(PtrType);
 							ralloc_asprintf_append(buffer, " *%s", unique_name(var));
@@ -791,10 +764,14 @@ protected:
 								auto& SamplerStates = SamplerStateFound->second;
 								for (auto& SamplerState : SamplerStates)
 								{
-									int32 SamplerStateIndex = Buffers.GetUniqueSamplerStateIndex(SamplerState, true);
-									ralloc_asprintf_append(
-										buffer,
-										"sampler s%d [[ sampler(%d) ]], ",  SamplerStateIndex, SamplerStateIndex);
+									bool bAdded = false;
+									int32 SamplerStateIndex = Buffers.GetUniqueSamplerStateIndex(SamplerState, true, bAdded);
+									if (bAdded)
+									{
+										ralloc_asprintf_append(
+											buffer,
+											"sampler s%d [[ sampler(%d) ]], ", SamplerStateIndex, SamplerStateIndex);
+									}
 								}
 							}
 
@@ -811,7 +788,7 @@ protected:
 					}
 					else
 					{
-						int BufferIndex = Buffers.GetIndex(var);
+						int BufferIndex = Buffers.GetIndex(var, Backend->bIsDesktop);
 						bool bNeedsPointer = (var->semantic && (strlen(var->semantic) == 1));
 						check(BufferIndex >= 0);
 						ralloc_asprintf_append(
@@ -892,6 +869,12 @@ protected:
 					{
 						ralloc_asprintf_append(buffer, "threadgroup ");
 					}
+
+					if (Buffers.AtomicVariables.find(var) != Buffers.AtomicVariables.end())
+					{
+						ralloc_asprintf_append(buffer, "atomic_");
+					}
+
 					print_type_pre(var->type);
 					ralloc_asprintf_append(buffer, " %s", unique_name(var));
 					print_type_post(var->type);
@@ -1015,6 +998,7 @@ protected:
 	virtual void visit(ir_expression *expr) override
 	{
 		check(scope_depth > 0);
+		++ExpressionDepth;
 
 		int numOps = expr->get_num_operands();
 		ir_expression_operation op = expr->operation;
@@ -1030,6 +1014,14 @@ protected:
 			expr->operands[0]->accept(this);
 			ralloc_asprintf_append(buffer, ")");
 		}
+		else if (op >= ir_unop_fasu && op <= ir_unop_uasf)
+		{
+			ralloc_asprintf_append(buffer, "as_type<");
+			print_type_full(expr->type);
+			ralloc_asprintf_append(buffer, ">(");
+			expr->operands[0]->accept(this);
+			ralloc_asprintf_append(buffer, ")");
+		}
 		else if (numOps == 1 && op >= ir_unop_first_conversion && op <= ir_unop_last_conversion)
 		{
 			std::string Type = expr->type->name;
@@ -1045,14 +1037,16 @@ protected:
 			 (numOps == 2 && op >= ir_binop_first_logic && op <= ir_binop_last_logic)))
 		{
 			const char* op_str = MetalExpressionTable[op][3];
-			ralloc_asprintf_append(buffer, "%s(", (numOps == 1) ? op_str : "");
+			ralloc_asprintf_append(buffer, "%s%s",
+				(numOps == 1) ? op_str : "",
+				(ExpressionDepth > 1) ? "(" : "");
 			expr->operands[0]->accept(this);
 			if (numOps == 2)
 			{
 				ralloc_asprintf_append(buffer, "%s", op_str);
 				expr->operands[1]->accept(this);
 			}
-			ralloc_asprintf_append(buffer, ")");
+			ralloc_asprintf_append(buffer, (ExpressionDepth > 1) ? ")" : "");
 		}
 		else if (expr->type->is_vector() && numOps == 2 &&
 			op >= ir_binop_first_logic && op <= ir_binop_last_logic)
@@ -1107,6 +1101,8 @@ protected:
 				ralloc_asprintf_append(buffer, MetalExpressionTable[op][i+1]);
 			}
 		}
+
+		--ExpressionDepth;
 	}
 
 	virtual void visit(ir_texture *tex) override
@@ -1115,7 +1111,7 @@ protected:
 		bool bNeedsClosingParenthesis = true;
 		if (tex->op == ir_txs)
 		{
-			ralloc_asprintf_append(buffer, "uint2(");
+			ralloc_asprintf_append(buffer, "int2((int)");
 		}
 
 		tex->sampler->accept(this);
@@ -1129,7 +1125,8 @@ protected:
 			auto* Texture = tex->sampler->variable_referenced();
 			check(Texture);
 			auto* Entry = ParseState->FindPackedSamplerEntry(Texture->name);
-			int32 SamplerStateIndex = Buffers.GetUniqueSamplerStateIndex(tex->SamplerStateName, false);
+			bool bDummy;
+			int32 SamplerStateIndex = Buffers.GetUniqueSamplerStateIndex(tex->SamplerStateName, false, bDummy);
 			check(SamplerStateIndex != INDEX_NONE);
 			ralloc_asprintf_append(buffer, "s%d, ", SamplerStateIndex);
 			if (tex->sampler->type->sampler_array)
@@ -1186,6 +1183,12 @@ protected:
 			{
 				ralloc_asprintf_append(buffer, ".read(");
 				tex->coordinate->accept(this);
+				
+				if (tex->sampler->type->sampler_ms)
+				{
+					ralloc_asprintf_append(buffer, ",");
+					tex->lod_info.sample_index->accept(this);
+				}
 			}
 		}
 			break;
@@ -1205,8 +1208,9 @@ protected:
 			// Sampler
 			auto* Texture = tex->sampler->variable_referenced();
 			check(Texture);
+			bool bDummy;
 			auto* Entry = ParseState->FindPackedSamplerEntry(Texture->name);
-			int32 SamplerStateIndex = Buffers.GetUniqueSamplerStateIndex(tex->SamplerStateName, false);
+			int32 SamplerStateIndex = Buffers.GetUniqueSamplerStateIndex(tex->SamplerStateName, false, bDummy);
 			check(SamplerStateIndex != INDEX_NONE);
 			ralloc_asprintf_append(buffer, "s%d, ", SamplerStateIndex);
 			// Coord
@@ -1230,19 +1234,19 @@ protected:
 		{
 			// Convert from:
 			//	HLSL
-			//		uint w, h;
+			//		int w, h;
 			//		T.GetDimensions({lod, }w, h);
 			// GLSL
 			//		ivec2 Temp;
 			//		Temp = textureSize(T{, lod});
 			// Metal
-			//		uint2 Temp = uint2(T.get_width({lod}), T.get_height({lod}));
+			//		int2 Temp = int2((int)T.get_width({lod}), (int)T.get_height({lod}));
 			ralloc_asprintf_append(buffer, ".get_width(");
 			if (tex->lod_info.lod)
 			{
 				tex->lod_info.lod->accept(this);
 			}
-			ralloc_asprintf_append(buffer, "), ");
+			ralloc_asprintf_append(buffer, "), (int)");
 			tex->sampler->accept(this);
 			ralloc_asprintf_append(buffer, ".get_height(");
 			if (tex->lod_info.lod)
@@ -1359,11 +1363,6 @@ protected:
 				hash_table_insert(used_uniform_blocks, (void*)var->semantic, var->semantic);
 			}
 		}
-
-		if (is_md_array(deref->type))
-		{
-			ralloc_asprintf_append(buffer, ".Inner");
-		}
 	}
 
 	virtual void visit(ir_dereference_array *deref) override
@@ -1407,11 +1406,6 @@ protected:
 		{
 			ralloc_asprintf_append(buffer, "]");
 		}
-
-		if (is_md_array(deref->array->type))
-		{
-			ralloc_asprintf_append(buffer, ".Inner");
-		}
 	}
 
 	void print_image_op( ir_dereference_image *deref, ir_rvalue *src)
@@ -1447,10 +1441,77 @@ protected:
 				if (bIsRWTexture)
 				{
 					ralloc_asprintf_append(buffer, ".write(");
-					src->accept(this);
-					ralloc_asprintf_append(buffer, ",");
+					// @todo Zebra: Below is a terrible hack - the input to write is always vec<T, 4>,
+					// 				but the type T comes from the texture type.  
+					if(src_elements == 1)
+					{
+						switch(deref->type->base_type)
+						{
+							case GLSL_TYPE_UINT:
+								ralloc_asprintf_append(buffer, "uint4(");
+								break;
+							case GLSL_TYPE_INT:
+								ralloc_asprintf_append(buffer, "int4(");
+								break;
+							case GLSL_TYPE_HALF:
+								ralloc_asprintf_append(buffer, "half4(");
+								break;
+							case GLSL_TYPE_FLOAT:
+							default:
+								ralloc_asprintf_append(buffer, "float4(");
+								break;
+						}
+						src->accept(this);
+						ralloc_asprintf_append(buffer, ")");
+					}
+					else
+					{
+						switch(deref->type->base_type)
+						{
+							case GLSL_TYPE_UINT:
+								ralloc_asprintf_append(buffer, "(uint4)(");
+								break;
+							case GLSL_TYPE_INT:
+								ralloc_asprintf_append(buffer, "(int4)(");
+								break;
+							case GLSL_TYPE_HALF:
+								ralloc_asprintf_append(buffer, "(half4)(");
+								break;
+							case GLSL_TYPE_FLOAT:
+							default:
+								ralloc_asprintf_append(buffer, "(float4)(");
+								break;
+						}
+						src->accept(this);
+						switch (src_elements)
+						{
+							case 3:
+								ralloc_asprintf_append(buffer, ").xyzx");
+								break;
+							case 2:
+								ralloc_asprintf_append(buffer, ").xyxy");
+								break;
+							default:
+								ralloc_asprintf_append(buffer, ")");
+								break;
+						}
+					}
+					
+					//#todo-rco: Add language spec to know if indices need to be uint
+					ralloc_asprintf_append(buffer, ",(uint");
+					switch (deref->image_index->type->vector_elements)
+					{
+					case 4:
+					case 3:
+					case 2:
+						ralloc_asprintf_append(buffer, "%d", deref->image_index->type->vector_elements);
+						//fallthrough
+					case 1:
+						ralloc_asprintf_append(buffer, ")(");
+						break;
+					}
 					deref->image_index->accept(this);
-					ralloc_asprintf_append(buffer, ")");
+					ralloc_asprintf_append(buffer, "))");
 				}
 				else
 				{
@@ -1488,11 +1549,6 @@ protected:
 
 		deref->record->accept(this);
 		ralloc_asprintf_append(buffer, ".%s", deref->field);
-
-		if (is_md_array(deref->type))
-		{
-			ralloc_asprintf_append(buffer, ".Inner");
-		}
 	}
 
 	virtual void visit(ir_assignment *assign) override
@@ -1800,7 +1856,9 @@ protected:
 		}
 
 		if (dest_deref == NULL || true_value == NULL)
+		{
 			return false;
+		}
 
 		num_inst = 0;
 		foreach_iter(exec_list_iterator, iter, expr->else_instructions)
@@ -1864,35 +1922,16 @@ protected:
 
 		if (try_conditional_move(expr) == false)
 		{
-		ralloc_asprintf_append(buffer, "if (");
-		expr->condition->accept(this);
-		ralloc_asprintf_append(buffer, ")\n");
-		indent();
-		ralloc_asprintf_append(buffer, "{\n");
-		
-		indentation++;
-		foreach_iter(exec_list_iterator, iter, expr->then_instructions)
-		{
-			ir_instruction *const inst = (ir_instruction *) iter.get();
-			indent();
-			do_visit(inst);
-		}
-		indentation--;
-
-		indent();
-		ralloc_asprintf_append(buffer, "}\n");
-
-		if (!expr->else_instructions.is_empty())
-		{
-			indent();
-			ralloc_asprintf_append(buffer, "else\n");
+			ralloc_asprintf_append(buffer, "if (");
+			expr->condition->accept(this);
+			ralloc_asprintf_append(buffer, ")\n");
 			indent();
 			ralloc_asprintf_append(buffer, "{\n");
 
 			indentation++;
-			foreach_iter(exec_list_iterator, iter, expr->else_instructions)
+			foreach_iter(exec_list_iterator, iter, expr->then_instructions)
 			{
-				ir_instruction *const inst = (ir_instruction *) iter.get();
+				ir_instruction *const inst = (ir_instruction *)iter.get();
 				indent();
 				do_visit(inst);
 			}
@@ -1900,10 +1939,29 @@ protected:
 
 			indent();
 			ralloc_asprintf_append(buffer, "}\n");
-		}
 
-		needs_semicolon = false;
-	}
+			if (!expr->else_instructions.is_empty())
+			{
+				indent();
+				ralloc_asprintf_append(buffer, "else\n");
+				indent();
+				ralloc_asprintf_append(buffer, "{\n");
+
+				indentation++;
+				foreach_iter(exec_list_iterator, iter, expr->else_instructions)
+				{
+					ir_instruction *const inst = (ir_instruction *)iter.get();
+					indent();
+					do_visit(inst);
+				}
+				indentation--;
+
+				indent();
+				ralloc_asprintf_append(buffer, "}\n");
+			}
+
+			needs_semicolon = false;
+		}
 	}
 
 	virtual void visit(ir_loop *loop) override
@@ -1968,16 +2026,17 @@ protected:
 	{
 		const char *sharedAtomicFunctions[] = 
 		{
-			"atomicAdd",
-			"atomicAnd",
-			"atomicMin",
-			"atomicMax",
-			"atomicOr",
-			"atomicXor",
-			"atomicExchange",
-			"atomicCompSwap"
+			"atomic_fetch_add_explicit",
+			"atomic_fetch_and_explicit",
+			"atomic_fetch_min_explicit",
+			"atomic_fetch_max_explicit",
+			"atomic_fetch_or_explicit",
+			"atomic_fetch_xor_explicit",
+			"atomic_exchange_explicit",
+			"atomic_compare_exchange_weak_explicit"
 		};
-		const char *imageAtomicFunctions[] = 
+/*
+		const char *imageAtomicFunctions[] =
 		{
 			"imageAtomicAdd",
 			"imageAtomicAnd",
@@ -1988,13 +2047,14 @@ protected:
 			"imageAtomicExchange",
 			"imageAtomicCompSwap"
 		};
+*/
 		check(scope_depth > 0);
 		const bool is_image = ir->memory_ref->as_dereference_image() != NULL;
 
 		ir->lhs->accept(this);
-		if (!is_image)
+		//if (!is_image)
 		{
-			ralloc_asprintf_append(buffer, " = %s(",
+			ralloc_asprintf_append(buffer, " = %s(&",
 				sharedAtomicFunctions[ir->operation]);
 			ir->memory_ref->accept(this);
 			ralloc_asprintf_append(buffer, ", ");
@@ -2004,8 +2064,9 @@ protected:
 				ralloc_asprintf_append(buffer, ", ");
 				ir->operands[1]->accept(this);
 			}
-			ralloc_asprintf_append(buffer, ")");
+			ralloc_asprintf_append(buffer, ", memory_order_relaxed)");
 		}
+/*
 		else
 		{
 			ir_dereference_image *image = ir->memory_ref->as_dereference_image();
@@ -2021,42 +2082,8 @@ protected:
 				ralloc_asprintf_append(buffer, ", ");
 				ir->operands[1]->accept(this);
 			}
-			ralloc_asprintf_append(buffer, ")");
-		}
-	}
-
-	/**
-	 * Declare structs used to simulate multi-dimensional arrays.
-	 */
-	void declare_md_array_struct(const glsl_type* type, hash_table* ht)
-	{
-		check(type->is_array());
-
-		if (hash_table_find(ht, (void*)type) == NULL)
-		{
-			const glsl_type* subtype = type->fields.array;
-			if (subtype->base_type == GLSL_TYPE_ARRAY)
-			{
-				declare_md_array_struct(subtype, ht);
-
-				ralloc_asprintf_append(buffer, "struct ");
-				print_md_array_type(type);
-				ralloc_asprintf_append(buffer, "\n{\n\t");
-				print_md_array_type(subtype);
-				ralloc_asprintf_append(buffer, " Inner[%u];\n};\n\n", type->length);
-			}
-			else
-			{
-				ralloc_asprintf_append(buffer, "struct ");
-				print_md_array_type(type);
-				ralloc_asprintf_append(buffer, "\n{\n\t");
-				print_type_pre(type);
-				ralloc_asprintf_append(buffer, " Inner");
-				print_type_post(type);
-				ralloc_asprintf_append(buffer, ";\n};\n\n");
-			}
-			hash_table_insert(ht, (void*)type, (void*)type);
-		}
+			ralloc_asprintf_append(buffer, ", memory_order_relaxed)");
+		}*/
 	}
 
 	/**
@@ -2130,17 +2157,6 @@ protected:
 			while( added_structure_types > 0 );
 		}
 
-		// Generate structures that allow support for multi-dimensional arrays.
-		{
-			hash_table* ht = hash_table_ctor(32, hash_table_pointer_hash, hash_table_pointer_compare);
-			foreach_iter(exec_list_iterator, iter, used_md_arrays) 
-			{
-				md_array_entry* entry = (md_array_entry*)iter.get();
-				declare_md_array_struct(entry->type, ht);
-			}
-			hash_table_dtor(ht);
-		}
-
 		for (unsigned i = 0; i < state->num_user_structures; i++)
 		{
 			const glsl_type *const s = state->user_structures[i];
@@ -2184,6 +2200,10 @@ protected:
 						{
 							ralloc_asprintf_append(buffer, " %s", s->fields.structure[j].semantic);
 						}
+						else if (!strcmp(s->fields.structure[j].semantic, "SV_RenderTargetArrayIndex"))
+						{
+							ralloc_asprintf_append(buffer, " [[ user(%s) ]]", NV_EPIC_LAYER);
+						}
 						else
 						{
 							ralloc_asprintf_append(buffer, "[[ ERROR! ]]");
@@ -2299,7 +2319,7 @@ protected:
 			glsl_packed_uniform& Uniform = *Iter;
 			ANSICHAR Name[32];
 			FCStringAnsi::Sprintf(Name, "%si%d", glsl_variable_tag_from_parser_target(Frequency), Uniform.offset);
-			int32 Offset = Buffers.GetIndex(Name);
+			int32 Offset = Buffers.GetIndex(Name, Backend->bIsDesktop);
 			check(Offset != -1);
 			ralloc_asprintf_append(
 				buffer,
@@ -2649,13 +2669,15 @@ protected:
 public:
 
 	/** Constructor. */
-	FGenerateMetalVisitor(_mesa_glsl_parse_state* InParseState, _mesa_glsl_parser_targets InFrequency, FBuffers& InBuffers)
-		: ParseState(InParseState)
+	FGenerateMetalVisitor(FMetalCodeBackend* InBackend, _mesa_glsl_parse_state* InParseState, _mesa_glsl_parser_targets InFrequency, FBuffers& InBuffers)
+		: Backend(InBackend)
+		, ParseState(InParseState)
 		, Frequency(InFrequency)
 		, Buffers(InBuffers)
 		, buffer(0)
 		, indentation(0)
 		, scope_depth(0)
+		, ExpressionDepth(0)
 		, temp_id(0)
 		, global_id(0)
 		, needs_semicolon(false)
@@ -2725,7 +2747,7 @@ char* FMetalCodeBackend::GenerateCode(exec_list* ir, _mesa_glsl_parse_state* sta
 {
 	// We'll need this Buffers info for the [[buffer()]] index
 	FBuffers Buffers;
-	FGenerateMetalVisitor visitor(state, state->target, Buffers);
+	FGenerateMetalVisitor visitor(this, state, state->target, Buffers);
 
 	// At this point, all inputs and outputs are global uniforms, no structures.
 
@@ -2761,6 +2783,8 @@ char* FMetalCodeBackend::GenerateCode(exec_list* ir, _mesa_glsl_parse_state* sta
 		//@todo-rco: Might not work if accessing packed_half* m[N]!
 		RemovePackedVarReferences(ir, state);
 
+		FindAtomicVariables(ir, Buffers.AtomicVariables);
+
 		bool bConvertUniformsToFloats = (HlslCompileFlags & HLSLCC_FlattenUniformBuffers) != HLSLCC_FlattenUniformBuffers;
 		ConvertHalfToFloatUniformsAndSamples(ir, state, bConvertUniformsToFloats, true);
 
@@ -2772,18 +2796,36 @@ char* FMetalCodeBackend::GenerateCode(exec_list* ir, _mesa_glsl_parse_state* sta
 	return _strdup(code);
 }
 
-struct FMetalUnsupportedSamplerVisitor : public ir_hierarchical_visitor
+struct FMetalCheckRestrictionsVisitor : public ir_hierarchical_visitor
 {
 	_mesa_glsl_parse_state* ParseState;
-	const bool bIsVertexShader;
-	FMetalUnsupportedSamplerVisitor(_mesa_glsl_parse_state* InParseState, bool bInIsVertexShader) :
+	EHlslShaderFrequency Frequency;
+	bool bErrors;
+	FMetalCheckRestrictionsVisitor(_mesa_glsl_parse_state* InParseState, EHlslShaderFrequency InFrequency) :
 		ParseState(InParseState),
-		bIsVertexShader(bInIsVertexShader)
+		bErrors(false),
+		Frequency(InFrequency)
 	{
 	}
 
-	virtual ir_visitor_status visit_leave(ir_texture* IR) override
+	virtual ir_visitor_status visit(ir_variable* IR) override
 	{
+		if (Frequency != HSF_ComputeShader)
+		{
+			if (IR->type && IR->type->is_image())
+			{
+				if (IR->name)
+				{
+					_mesa_glsl_error(ParseState, "Metal doesn't allow UAV '%s' on non-compute shader stages.", IR->name);
+				}
+				else
+				{
+					_mesa_glsl_error(ParseState, "Metal doesn't allow UAV on non-compute shader stages.");
+				}
+				bErrors = true;
+				return visit_stop;
+			}
+		}
 
 		return visit_continue;
 	}
@@ -2792,16 +2834,10 @@ struct FMetalUnsupportedSamplerVisitor : public ir_hierarchical_visitor
 
 bool FMetalCodeBackend::ApplyAndVerifyPlatformRestrictions(exec_list* Instructions, _mesa_glsl_parse_state* ParseState, EHlslShaderFrequency Frequency)
 {
-	bool bIsVertexShader = (Frequency == HSF_VertexShader);
+	FMetalCheckRestrictionsVisitor Visitor(ParseState, Frequency);
+	Visitor.run(Instructions);
 
-	const bool bAllowVertexTextureFetch = true;
-	if (!bAllowVertexTextureFetch)
-	{
-		FMetalUnsupportedSamplerVisitor Visitor(ParseState, bIsVertexShader);
-		Visitor.run(Instructions);
-	}
-
-	return true;
+	return !Visitor.bErrors;
 }
 
 bool FMetalCodeBackend::GenerateMain(EHlslShaderFrequency Frequency, const char* EntryPoint, exec_list* Instructions, _mesa_glsl_parse_state* ParseState)
@@ -2831,7 +2867,7 @@ bool FMetalCodeBackend::GenerateMain(EHlslShaderFrequency Frequency, const char*
 			{
 			case ir_var_in:
 				ArgVarDeref = MetalUtils::GenerateInput(
-					Frequency,
+					Frequency, bIsDesktop,
 					ParseState,
 					Variable->semantic,
 					Variable->type,
@@ -2841,7 +2877,7 @@ bool FMetalCodeBackend::GenerateMain(EHlslShaderFrequency Frequency, const char*
 				break;
 			case ir_var_out:
 				ArgVarDeref = MetalUtils::GenerateOutput(
-					Frequency,
+					Frequency, bIsDesktop,
 					ParseState,
 					Variable->semantic,
 					Variable->type,
@@ -2867,7 +2903,7 @@ bool FMetalCodeBackend::GenerateMain(EHlslShaderFrequency Frequency, const char*
 	ir_dereference_variable* EntryPointReturn = nullptr;
 	if (!EntryPointSig->return_type->is_void())
 	{
-		EntryPointReturn = MetalUtils::GenerateOutput(Frequency, ParseState, EntryPointSig->return_semantic, EntryPointSig->return_type, &DeclInstructions, &PreCallInstructions, &PostCallInstructions);
+		EntryPointReturn = MetalUtils::GenerateOutput(Frequency, bIsDesktop, ParseState, EntryPointSig->return_semantic, EntryPointSig->return_type, &DeclInstructions, &PreCallInstructions, &PostCallInstructions);
 	}
 
 	ParseState->symbols->pop_scope();
@@ -2900,6 +2936,7 @@ bool FMetalCodeBackend::GenerateMain(EHlslShaderFrequency Frequency, const char*
 FMetalCodeBackend::FMetalCodeBackend(unsigned int InHlslCompileFlags, EHlslCompileTarget InTarget) :
 	FCodeBackend(InHlslCompileFlags, InTarget)
 {
+	bIsDesktop = (InTarget == HCT_FeatureLevelSM4 || InTarget == HCT_FeatureLevelSM5);
 }
 
 void FMetalLanguageSpec::SetupLanguageIntrinsics(_mesa_glsl_parse_state* State, exec_list* ir)
diff --git a/Engine/Source/Developer/iOS/MetalShaderFormat/Private/MetalBackend.h b/Engine/Source/Developer/iOS/MetalShaderFormat/Private/MetalBackend.h
--- a/Engine/Source/Developer/iOS/MetalShaderFormat/Private/MetalBackend.h
+++ b/Engine/Source/Developer/iOS/MetalShaderFormat/Private/MetalBackend.h
@@ -52,4 +52,6 @@ struct FMetalCodeBackend : public FCodeBackend
 	void PromoteInputsAndOutputsGlobalHalfToFloat(exec_list* ir, _mesa_glsl_parse_state* state, EHlslShaderFrequency Frequency);
 	void ConvertHalfToFloatUniformsAndSamples(exec_list* ir, _mesa_glsl_parse_state* State, bool bConvertUniforms, bool bConvertSamples);
 	void BreakPrecisionChangesVisitor(exec_list* ir, _mesa_glsl_parse_state* State);
+
+	bool bIsDesktop;
 };
diff --git a/Engine/Source/Developer/iOS/MetalShaderFormat/Private/MetalShaderCompiler.cpp b/Engine/Source/Developer/iOS/MetalShaderFormat/Private/MetalShaderCompiler.cpp
--- a/Engine/Source/Developer/iOS/MetalShaderFormat/Private/MetalShaderCompiler.cpp
+++ b/Engine/Source/Developer/iOS/MetalShaderFormat/Private/MetalShaderCompiler.cpp
@@ -86,6 +86,7 @@ static void BuildMetalShaderOutput(
 	const FShaderCompilerInput& ShaderInput, 
 	const ANSICHAR* InShaderSource,
 	int32 SourceLen,
+	bool bIsMobile,
 	TArray<FShaderCompilerError>& OutErrors
 	)
 {
@@ -117,6 +118,14 @@ static void BuildMetalShaderOutput(
 	DEF_PREFIX_STR(NumThreads);
 	#undef DEF_PREFIX_STR
 
+	ANSICHAR const* ShaderName = nullptr;
+	uint32 ShaderNameLen = 0;
+	while (	FCStringAnsi::Strncmp(ShaderSource, "// !", 4) == 0 )
+	{
+		ShaderName = ShaderSource;
+		while (*ShaderSource && *ShaderSource++ != '\n') { ShaderNameLen++; }
+	}
+	
 	// Skip any comments that come before the signature.
 	while (	FCStringAnsi::Strncmp(ShaderSource, "//", 2) == 0 &&
 			FCStringAnsi::Strncmp(ShaderSource, "// @", 4) != 0 )
@@ -505,6 +514,7 @@ static void BuildMetalShaderOutput(
 	}
 
 	// Then samplers.
+	TMap<FString, uint32> SamplerMap;
 	if (FCStringAnsi::Strncmp(ShaderSource, SamplersPrefix, SamplersPrefixLen) == 0)
 	{
 		ShaderSource += SamplersPrefixLen;
@@ -539,12 +549,7 @@ static void BuildMetalShaderOutput(
 				{
 					FString SamplerState = ParseIdentifier(ShaderSource);
 					checkSlow(SamplerState.Len() != 0);
-					ParameterMap.AddParameterAllocation(
-						*SamplerState,
-						0,
-						Offset,
-						NumSamplers
-						);
+					SamplerMap.Add(SamplerState, NumSamplers);
 				}
 				while (Match(ShaderSource, ','));
 				verify(Match(ShaderSource, ']'));
@@ -612,6 +617,14 @@ static void BuildMetalShaderOutput(
 			verify(Match(ShaderSource, ':'));
 			FString SamplerState = ParseIdentifier(ShaderSource);
 
+			checkSlow(SamplerState.Len() != 0);
+			ParameterMap.AddParameterAllocation(
+				*SamplerState,
+				0,
+				Index,
+				SamplerMap[SamplerState]
+				);
+			
 			if (Match(ShaderSource, '\n'))
 			{
 				break;
@@ -665,9 +678,25 @@ static void BuildMetalShaderOutput(
 			FString::Printf(TEXT("shader uses %d samplers exceeding the limit of %d"),
 				Header.Bindings.NumSamplers, MaxSamplers);
 	}
+	else if(ShaderInput.Environment.CompilerFlags.Contains(CFLAG_Debug))
+	{
+		// Write out the header and shader source code.
+		FMemoryWriter Ar(ShaderOutput.Code, true);
+		uint8 PrecompiledFlag = 0;
+		Ar << PrecompiledFlag;
+		Ar << Header;
+		if(ShaderName && ShaderNameLen)
+		{
+			Ar.Serialize((void*)ShaderName, ShaderNameLen);
+		}
+		Ar.Serialize((void*)ShaderSource, SourceLen + 1 - (ShaderSource - InShaderSource));
+		
+		ShaderOutput.NumInstructions = 0;
+		ShaderOutput.NumTextureSamplers = Header.Bindings.NumSamplers;
+		ShaderOutput.bSucceeded = true;
+	}
 	else
 	{
-#if METAL_OFFLINE_COMPILE
 		// at this point, the shader source is ready to be compiled
 		FString InputFilename = FPaths::CreateTempFilename(*FPaths::EngineIntermediateDir(), TEXT("ShaderIn"), TEXT(""));
 		FString ObjFilename = InputFilename + TEXT(".o");
@@ -681,49 +710,56 @@ static void BuildMetalShaderOutput(
 		int32 ReturnCode = 0;
 		FString Results;
 		FString Errors;
-		bool bHadError = true;
+		bool bCompileAtRuntime = true;
+		bool bSucceeded = false;
 
+#if PLATFORM_MAC
 		// metal commandlines
-		FString Params = FString::Printf(TEXT("-std=ios-metal1.0 %s -o %s"), *InputFilename, *ObjFilename);
-		FPlatformProcess::ExecProcess( TEXT("/Applications/Xcode.app/Contents/Developer/Platforms/iPhoneOS.platform/usr/bin/metal"), *Params, &ReturnCode, &Results, &Errors );
+		FString Standard = (bIsMobile ? TEXT("-std=ios-metal1.0") : TEXT("-std=osx-metal1.1 -Wno-null-character -ffmast-math"));
+		
+		FString Params = FString::Printf(TEXT("%s %s -o %s"), *Standard, *InputFilename, *ObjFilename);
+		FPlatformProcess::ExecProcess( TEXT("/Applications/Xcode.app/Contents/Developer/Toolchains/XcodeDefault.xctoolchain/usr/bin/metal"), *Params, &ReturnCode, &Results, &Errors );
 
 		// handle compile error
 		if (ReturnCode != 0 || IFileManager::Get().FileSize(*ObjFilename) <= 0)
 		{
-//			FShaderCompilerError* Error = new(OutErrors) FShaderCompilerError();
-//			Error->ErrorFile = InputFilename;
-//			Error->ErrorLineString = TEXT("0");
-//			Error->StrippedErrorMessage = Results + Errors;
+			FShaderCompilerError* Error = new(OutErrors) FShaderCompilerError();
+			Error->ErrorFile = InputFilename;
+			Error->ErrorLineString = TEXT("0");
+			Error->StrippedErrorMessage = Results + Errors;
+			bSucceeded = false;
 		}
 		else
 		{
 			Params = FString::Printf(TEXT("r %s %s"), *ArFilename, *ObjFilename);
-			FPlatformProcess::ExecProcess( TEXT("/Applications/Xcode.app/Contents/Developer/Platforms/iPhoneOS.platform/usr/bin/metal-ar"), *Params, &ReturnCode, &Results, &Errors );
+			FPlatformProcess::ExecProcess( TEXT("/Applications/Xcode.app/Contents/Developer/Toolchains/XcodeDefault.xctoolchain/usr/bin/metal-ar"), *Params, &ReturnCode, &Results, &Errors );
 
 			// handle compile error
 			if (ReturnCode != 0 || IFileManager::Get().FileSize(*ArFilename) <= 0)
 			{
-//				FShaderCompilerError* Error = new(OutErrors) FShaderCompilerError();
-//				Error->ErrorFile = InputFilename;
-//				Error->ErrorLineString = TEXT("0");
-//				Error->StrippedErrorMessage = Results + Errors;
+				FShaderCompilerError* Error = new(OutErrors) FShaderCompilerError();
+				Error->ErrorFile = InputFilename;
+				Error->ErrorLineString = TEXT("0");
+				Error->StrippedErrorMessage = Results + Errors;
+				bSucceeded = false;
 			}
 			else
 			{
 				Params = FString::Printf(TEXT("-o %s %s"), *OutputFilename, *ArFilename);
-				FPlatformProcess::ExecProcess( TEXT("/Applications/Xcode.app/Contents/Developer/Platforms/iPhoneOS.platform/usr/bin/metallib"), *Params, &ReturnCode, &Results, &Errors );
+				FPlatformProcess::ExecProcess( TEXT("/Applications/Xcode.app/Contents/Developer/Toolchains/XcodeDefault.xctoolchain/usr/bin/metallib"), *Params, &ReturnCode, &Results, &Errors );
 		
 				// handle compile error
 				if (ReturnCode != 0 || IFileManager::Get().FileSize(*OutputFilename) <= 0)
 				{
-//					FShaderCompilerError* Error = new(OutErrors) FShaderCompilerError();
-//					Error->ErrorFile = InputFilename;
-//					Error->ErrorLineString = TEXT("0");
-//					Error->StrippedErrorMessage = Results + Errors;
+					FShaderCompilerError* Error = new(OutErrors) FShaderCompilerError();
+					Error->ErrorFile = InputFilename;
+					Error->ErrorLineString = TEXT("0");
+					Error->StrippedErrorMessage = Results + Errors;
+					bSucceeded = false;
 				}
 				else
 				{
-					bHadError = false;
+					bCompileAtRuntime = false;
 					
 					// Write out the header and compiled shader code
 					FMemoryWriter Ar(ShaderOutput.Code, true);
@@ -741,11 +777,16 @@ static void BuildMetalShaderOutput(
 					ShaderOutput.NumInstructions = 0;
 					ShaderOutput.NumTextureSamplers = Header.Bindings.NumSamplers;
 					ShaderOutput.bSucceeded = true;
+					bSucceeded = true;
 				}
 			}
 		}
-		
-		if (bHadError)
+#else
+		// Assume success for non-Mac
+		bSucceeded = true;
+#endif	// PLATFORM_MAC
+
+		if (bCompileAtRuntime)
 		{
 			// Write out the header and shader source code.
 			FMemoryWriter Ar(ShaderOutput.Code, true);
@@ -756,25 +797,15 @@ static void BuildMetalShaderOutput(
 			
 			ShaderOutput.NumInstructions = 0;
 			ShaderOutput.NumTextureSamplers = Header.Bindings.NumSamplers;
-			ShaderOutput.bSucceeded = true;
+			ShaderOutput.bSucceeded = bSucceeded || ShaderOutput.bSucceeded;
+		}
+		else
+		{
+			IFileManager::Get().Delete(*InputFilename);
+			IFileManager::Get().Delete(*ObjFilename);
+			IFileManager::Get().Delete(*ArFilename);
+			IFileManager::Get().Delete(*OutputFilename);
 		}
-		
-		IFileManager::Get().Delete(*InputFilename);
-		IFileManager::Get().Delete(*ObjFilename);
-		IFileManager::Get().Delete(*ArFilename);
-		IFileManager::Get().Delete(*OutputFilename);
-#else
-		// Write out the header and shader source code.
-		FMemoryWriter Ar(ShaderOutput.Code, true);
-		uint8 PrecompiledFlag = 0;
-		Ar << PrecompiledFlag;
-		Ar << Header;
-		Ar.Serialize((void*)ShaderSource, SourceLen + 1 - (ShaderSource - InShaderSource));
-		
-		ShaderOutput.NumInstructions = 0;
-		ShaderOutput.NumTextureSamplers = Header.Bindings.NumSamplers;
-		ShaderOutput.bSucceeded = true;
-#endif
 	}
 }
 
@@ -853,11 +884,24 @@ void CompileShader_Metal(const FShaderCompilerInput& Input,FShaderCompilerOutput
 	FString PreprocessedShader;
 	FShaderCompilerDefinitions AdditionalDefines;
 	EHlslCompileTarget HlslCompilerTarget = HCT_FeatureLevelES3_1;
-	AdditionalDefines.SetDefine(TEXT("IOS"), 1);
+
+	// @todo - Zebra - Work out which standard we need, this is dependent on the shader platform.
+	// For now only SP_METAL will compile for iOS, with MRT & SM5 only for Mac.
+	const bool bIsMobile = (Output.Target.Platform == SP_METAL);
+	if (bIsMobile)
+	{
+		AdditionalDefines.SetDefine(TEXT("IOS"), 1);
+	}
+	else
+	{
+		AdditionalDefines.SetDefine(TEXT("MAC"), 1);
+	}
+
 	AdditionalDefines.SetDefine(TEXT("row_major"), TEXT(""));
 
 	static FName NAME_SF_METAL(TEXT("SF_METAL"));
 	static FName NAME_SF_METAL_MRT(TEXT("SF_METAL_MRT"));
+	static FName NAME_SF_METAL_SM5(TEXT("SF_METAL_SM5"));
 
 	if (Input.ShaderFormat == NAME_SF_METAL)
 	{
@@ -867,6 +911,10 @@ void CompileShader_Metal(const FShaderCompilerInput& Input,FShaderCompilerOutput
 	{
 		AdditionalDefines.SetDefine(TEXT("METAL_MRT_PROFILE"), 1);
 	}
+	else if (Input.ShaderFormat == NAME_SF_METAL_SM5)
+	{
+		AdditionalDefines.SetDefine(TEXT("METAL_SM5_PROFILE"), 1);
+	}
 	else
 	{
 		Output.bSucceeded = false;
@@ -876,7 +924,15 @@ void CompileShader_Metal(const FShaderCompilerInput& Input,FShaderCompilerOutput
 	
 	const bool bDumpDebugInfo = (Input.DumpDebugInfoPath != TEXT("") && IFileManager::Get().DirectoryExists(*Input.DumpDebugInfoPath));
 
-	AdditionalDefines.SetDefine(TEXT("COMPILER_SUPPORTS_ATTRIBUTES"), (uint32)1);
+	if(Input.Environment.CompilerFlags.Contains(CFLAG_AvoidFlowControl))
+	{
+		AdditionalDefines.SetDefine(TEXT("COMPILER_SUPPORTS_ATTRIBUTES"), (uint32)1);
+	}
+	else
+	{
+		AdditionalDefines.SetDefine(TEXT("COMPILER_SUPPORTS_ATTRIBUTES"), (uint32)0);
+	}
+
 	if (PreprocessShader(PreprocessedShader, Output, Input, AdditionalDefines))
 	{
 		char* MetalShaderSource = NULL;
@@ -925,6 +981,10 @@ void CompileShader_Metal(const FShaderCompilerInput& Input,FShaderCompilerOutput
 		}
 
 		uint32 CCFlags = HLSLCC_NoPreprocess | HLSLCC_PackUniforms;
+		if (!bIsMobile)
+		{
+			CCFlags |= HLSLCC_FixAtomicReferences;
+		}
 		//CCFlags |= HLSLCC_FlattenUniformBuffers | HLSLCC_FlattenUniformBufferStructures;
 
 		if (bDumpDebugInfo)
@@ -958,6 +1018,15 @@ void CompileShader_Metal(const FShaderCompilerInput& Input,FShaderCompilerOutput
 		}
 
 		int32 SourceLen = MetalShaderSource ? FCStringAnsi::Strlen(MetalShaderSource) : 0;
+		if(MetalShaderSource)
+		{
+			uint32 Len = FCStringAnsi::Strlen(TCHAR_TO_ANSI(*Input.SourceFilename)) + FCStringAnsi::Strlen(TCHAR_TO_ANSI(*Input.EntryPointName)) + FCStringAnsi::Strlen(MetalShaderSource) + 20;
+			char* Dest = (char*)malloc(Len);
+			sprintf(Dest, "// ! %s.usf:%s\n%s", TCHAR_TO_ANSI(*Input.SourceFilename), TCHAR_TO_ANSI(*Input.EntryPointName), MetalShaderSource);
+			free(MetalShaderSource);
+			MetalShaderSource = Dest;
+			SourceLen = FCStringAnsi::Strlen(MetalShaderSource);
+		}
 		if (bDumpDebugInfo)
 		{
 			if (SourceLen > 0)
@@ -975,7 +1044,7 @@ void CompileShader_Metal(const FShaderCompilerInput& Input,FShaderCompilerOutput
 		if (Result != 0)
 		{
 			Output.Target = Input.Target;
-			BuildMetalShaderOutput(Output, Input, MetalShaderSource, SourceLen, Output.Errors);
+			BuildMetalShaderOutput(Output, Input, MetalShaderSource, SourceLen, bIsMobile, Output.Errors);
 		}
 		else
 		{
diff --git a/Engine/Source/Developer/iOS/MetalShaderFormat/Private/MetalShaderFormat.cpp b/Engine/Source/Developer/iOS/MetalShaderFormat/Private/MetalShaderFormat.cpp
--- a/Engine/Source/Developer/iOS/MetalShaderFormat/Private/MetalShaderFormat.cpp
+++ b/Engine/Source/Developer/iOS/MetalShaderFormat/Private/MetalShaderFormat.cpp
@@ -10,12 +10,13 @@
 
 static FName NAME_SF_METAL(TEXT("SF_METAL"));
 static FName NAME_SF_METAL_MRT(TEXT("SF_METAL_MRT"));
+static FName NAME_SF_METAL_SM5(TEXT("SF_METAL_SM5"));
 
 class FMetalShaderFormat : public IShaderFormat
 {
 	enum
 	{
-		HEADER_VERSION = 14,
+		HEADER_VERSION = 17,
 	};
 	struct FVersion
 	{
@@ -51,10 +52,11 @@ public:
 	{
 		OutFormats.Add(NAME_SF_METAL);
 		OutFormats.Add(NAME_SF_METAL_MRT);
+		OutFormats.Add(NAME_SF_METAL_SM5);
 	}
 	virtual void CompileShader(FName Format, const struct FShaderCompilerInput& Input, struct FShaderCompilerOutput& Output,const FString& WorkingDirectory) const
 	{
-		check(Format == NAME_SF_METAL || Format == NAME_SF_METAL_MRT);
+		check(Format == NAME_SF_METAL || Format == NAME_SF_METAL_MRT || Format == NAME_SF_METAL_SM5);
 		CompileShader_Metal(Input, Output, WorkingDirectory);
 	}
 };
diff --git a/Engine/Source/Developer/iOS/MetalShaderFormat/Private/MetalShaderFormat.h b/Engine/Source/Developer/iOS/MetalShaderFormat/Private/MetalShaderFormat.h
--- a/Engine/Source/Developer/iOS/MetalShaderFormat/Private/MetalShaderFormat.h
+++ b/Engine/Source/Developer/iOS/MetalShaderFormat/Private/MetalShaderFormat.h
@@ -6,5 +6,5 @@ extern void CompileShader_Metal(const struct FShaderCompilerInput& Input, struct
 
 // Set this to 0 to get shader source in the graphics debugger
 // Note: Offline and runtime compiled shaders have separate DDC versions and can co-exist
-#define METAL_OFFLINE_COMPILE !PLATFORM_WINDOWS
+#define METAL_OFFLINE_COMPILE 0 // !PLATFORM_WINDOWS
 
diff --git a/Engine/Source/Developer/iOS/MetalShaderFormat/Private/MetalUtils.cpp b/Engine/Source/Developer/iOS/MetalShaderFormat/Private/MetalUtils.cpp
--- a/Engine/Source/Developer/iOS/MetalShaderFormat/Private/MetalUtils.cpp
+++ b/Engine/Source/Developer/iOS/MetalShaderFormat/Private/MetalUtils.cpp
@@ -183,17 +183,18 @@ namespace MetalUtils
 		{"SV_VertexID", glsl_type::uint_type, "IN_VertexID", ir_var_in, "[[ vertex_id ]]"},
 		{"SV_InstanceID", glsl_type::uint_type, "IN_InstanceID", ir_var_in, "[[ instance_id ]]"},
 		{"SV_Position", glsl_type::vec4_type, "Position", ir_var_out, "[[ position ]]"},
+		{"SV_RenderTargetArrayIndex", glsl_type::uint_type, "OUT_Layer", ir_var_out, "SV_RenderTargetArrayIndex"},
 		{NULL, NULL, NULL, ir_var_auto, nullptr}
 	};
 
 	/** Pixel shader system values. */
-	static FSystemValue PixelSystemValueTable[] =
+	static FSystemValue MobilePixelSystemValueTable[] =
 	{
 		{"SV_Depth", glsl_type::float_type, "FragDepth", ir_var_out, "[[ depth(any) ]]"},
 		{"SV_Position", glsl_type::vec4_type, "IN_FragCoord", ir_var_in, "[[ position ]]"},
 		{"SV_IsFrontFace", glsl_type::bool_type, "IN_FrontFacing", ir_var_in, "[[ front_facing ]]"},
 		//{"SV_PrimitiveID", glsl_type::int_type, "IN_PrimitiveID", ir_var_in, "[[  ]]"},
-		//{"SV_RenderTargetArrayIndex", glsl_type::int_type, "IN_Layer", ir_var_in, nullptr},
+		//{"SV_RenderTargetArrayIndex", glsl_type::uint_type, "IN_Layer", ir_var_in, "SV_RenderTargetArrayIndex"},
 		{"SV_Target0", glsl_type::half4_type, "FragColor0", ir_var_out, "[[ color(0) ]]"},
 		{"SV_Target1", glsl_type::half4_type, "FragColor1", ir_var_out, "[[ color(1) ]]"},
 		{"SV_Target2", glsl_type::half4_type, "FragColor2", ir_var_out, "[[ color(2) ]]"},
@@ -218,14 +219,57 @@ namespace MetalUtils
 		{"SV_Target5", glsl_type::half2_type, "FragColor5", ir_var_out, "[[ color(5) ]]"},
 		{"SV_Target6", glsl_type::half2_type, "FragColor6", ir_var_out, "[[ color(6) ]]"},
 		{"SV_Target7", glsl_type::half2_type, "FragColor7", ir_var_out, "[[ color(7) ]]"},
-		{"SV_Target0", glsl_type::half_type, "FragColor0", ir_var_out, "[[ color(0) ]]"},
-		{"SV_Target1", glsl_type::half_type, "FragColor1", ir_var_out, "[[ color(1) ]]"},
-		{"SV_Target2", glsl_type::half_type, "FragColor2", ir_var_out, "[[ color(2) ]]"},
-		{"SV_Target3", glsl_type::half_type, "FragColor3", ir_var_out, "[[ color(3) ]]"},
-		{"SV_Target4", glsl_type::half_type, "FragColor4", ir_var_out, "[[ color(4) ]]"},
-		{"SV_Target5", glsl_type::half_type, "FragColor5", ir_var_out, "[[ color(5) ]]"},
-		{"SV_Target6", glsl_type::half_type, "FragColor6", ir_var_out, "[[ color(6) ]]"},
-		{"SV_Target7", glsl_type::half_type, "FragColor7", ir_var_out, "[[ color(7) ]]"},
+		{"SV_Target0", glsl_type::float_type, "FragColor0", ir_var_out, "[[ color(0) ]]"},
+		{"SV_Target1", glsl_type::float_type, "FragColor1", ir_var_out, "[[ color(1) ]]"},
+		{"SV_Target2", glsl_type::float_type, "FragColor2", ir_var_out, "[[ color(2) ]]"},
+		{"SV_Target3", glsl_type::float_type, "FragColor3", ir_var_out, "[[ color(3) ]]"},
+		{"SV_Target4", glsl_type::float_type, "FragColor4", ir_var_out, "[[ color(4) ]]"},
+		{"SV_Target5", glsl_type::float_type, "FragColor5", ir_var_out, "[[ color(5) ]]"},
+		{"SV_Target6", glsl_type::float_type, "FragColor6", ir_var_out, "[[ color(6) ]]"},
+		{"SV_Target7", glsl_type::float_type, "FragColor7", ir_var_out, "[[ color(7) ]]"},
+		{NULL, NULL, NULL, ir_var_auto, nullptr}
+	};
+
+	/** Pixel shader system values. */
+	static FSystemValue DesktopPixelSystemValueTable[] =
+	{
+		{"SV_Depth", glsl_type::float_type, "FragDepth", ir_var_out, "[[ depth(any) ]]"},
+		{"SV_Position", glsl_type::vec4_type, "IN_FragCoord", ir_var_in, "[[ position ]]"},
+		{"SV_IsFrontFace", glsl_type::bool_type, "IN_FrontFacing", ir_var_in, "[[ front_facing ]]"},
+		//{"SV_PrimitiveID", glsl_type::int_type, "IN_PrimitiveID", ir_var_in, "[[  ]]"},
+		//{"SV_RenderTargetArrayIndex", glsl_type::uint_type, "IN_Layer", ir_var_in, "SV_RenderTargetArrayIndex"},
+		{"SV_Target0", glsl_type::vec4_type, "FragColor0", ir_var_out, "[[ color(0) ]]"},
+		{"SV_Target1", glsl_type::vec4_type, "FragColor1", ir_var_out, "[[ color(1) ]]"},
+		{"SV_Target2", glsl_type::vec4_type, "FragColor2", ir_var_out, "[[ color(2) ]]"},
+		{"SV_Target3", glsl_type::vec4_type, "FragColor3", ir_var_out, "[[ color(3) ]]"},
+		{"SV_Target4", glsl_type::vec4_type, "FragColor4", ir_var_out, "[[ color(4) ]]"},
+		{"SV_Target5", glsl_type::vec4_type, "FragColor5", ir_var_out, "[[ color(5) ]]"},
+		{"SV_Target6", glsl_type::vec4_type, "FragColor6", ir_var_out, "[[ color(6) ]]"},
+		{"SV_Target7", glsl_type::vec4_type, "FragColor7", ir_var_out, "[[ color(7) ]]"},
+		{"SV_Target0", glsl_type::vec3_type, "FragColor0", ir_var_out, "[[ color(0) ]]"},
+		{"SV_Target1", glsl_type::vec3_type, "FragColor1", ir_var_out, "[[ color(1) ]]"},
+		{"SV_Target2", glsl_type::vec3_type, "FragColor2", ir_var_out, "[[ color(2) ]]"},
+		{"SV_Target3", glsl_type::vec3_type, "FragColor3", ir_var_out, "[[ color(3) ]]"},
+		{"SV_Target4", glsl_type::vec3_type, "FragColor4", ir_var_out, "[[ color(4) ]]"},
+		{"SV_Target5", glsl_type::vec3_type, "FragColor5", ir_var_out, "[[ color(5) ]]"},
+		{"SV_Target6", glsl_type::vec3_type, "FragColor6", ir_var_out, "[[ color(6) ]]"},
+		{"SV_Target7", glsl_type::vec3_type, "FragColor7", ir_var_out, "[[ color(7) ]]"},
+		{"SV_Target0", glsl_type::vec2_type, "FragColor0", ir_var_out, "[[ color(0) ]]"},
+		{"SV_Target1", glsl_type::vec2_type, "FragColor1", ir_var_out, "[[ color(1) ]]"},
+		{"SV_Target2", glsl_type::vec2_type, "FragColor2", ir_var_out, "[[ color(2) ]]"},
+		{"SV_Target3", glsl_type::vec2_type, "FragColor3", ir_var_out, "[[ color(3) ]]"},
+		{"SV_Target4", glsl_type::vec2_type, "FragColor4", ir_var_out, "[[ color(4) ]]"},
+		{"SV_Target5", glsl_type::vec2_type, "FragColor5", ir_var_out, "[[ color(5) ]]"},
+		{"SV_Target6", glsl_type::vec2_type, "FragColor6", ir_var_out, "[[ color(6) ]]"},
+		{"SV_Target7", glsl_type::vec2_type, "FragColor7", ir_var_out, "[[ color(7) ]]"},
+		{"SV_Target0", glsl_type::float_type, "FragColor0", ir_var_out, "[[ color(0) ]]"},
+		{"SV_Target1", glsl_type::float_type, "FragColor1", ir_var_out, "[[ color(1) ]]"},
+		{"SV_Target2", glsl_type::float_type, "FragColor2", ir_var_out, "[[ color(2) ]]"},
+		{"SV_Target3", glsl_type::float_type, "FragColor3", ir_var_out, "[[ color(3) ]]"},
+		{"SV_Target4", glsl_type::float_type, "FragColor4", ir_var_out, "[[ color(4) ]]"},
+		{"SV_Target5", glsl_type::float_type, "FragColor5", ir_var_out, "[[ color(5) ]]"},
+		{"SV_Target6", glsl_type::float_type, "FragColor6", ir_var_out, "[[ color(6) ]]"},
+		{"SV_Target7", glsl_type::float_type, "FragColor7", ir_var_out, "[[ color(7) ]]"},
 		{NULL, NULL, NULL, ir_var_auto, nullptr}
 	};
 
@@ -237,7 +281,7 @@ namespace MetalUtils
 		{"SV_InstanceID", glsl_type::int_type, "IN_InstanceID", ir_var_in, false, false, false, false},
 		{"SV_Position", glsl_type::vec4_type, "IN_Position", ir_var_in, false, true, true, false},
 		{"SV_Position", glsl_type::vec4_type, "OUT_Position", ir_var_out, false, false, true, false},
-		{"SV_RenderTargetArrayIndex", glsl_type::int_type, "OUT_Layer", ir_var_out, false, false, false, false},
+		{"SV_RenderTargetArrayIndex", glsl_type::uint_type, "OUT_Layer", ir_var_out, false, false, false, false},
 		{"SV_PrimitiveID", glsl_type::int_type, "OUT_PrimitiveID", ir_var_out, false, false, false, false},
 		{"SV_PrimitiveID", glsl_type::int_type, "IN_PrimitiveIDIn", ir_var_in, false, false, false, false},
 */
@@ -277,18 +321,27 @@ namespace MetalUtils
 		{NULL, NULL, NULL, ir_var_auto, nullptr}
 	};
 
-	FSystemValue* SystemValueTable[HSF_FrequencyCount] =
+	FSystemValue* MobileSystemValueTable[HSF_FrequencyCount] =
 	{
 		VertexSystemValueTable,
-		PixelSystemValueTable,
+		MobilePixelSystemValueTable,
 		GeometrySystemValueTable,
 		HullSystemValueTable,
 		DomainSystemValueTable,
 		ComputeSystemValueTable
 	};
 
+	FSystemValue* DesktopSystemValueTable[HSF_FrequencyCount] =
+	{
+		VertexSystemValueTable,
+		DesktopPixelSystemValueTable,
+		GeometrySystemValueTable,
+		HullSystemValueTable,
+		DomainSystemValueTable,
+		ComputeSystemValueTable
+	};
 
-	static ir_rvalue* GenerateInputFromSemantic(EHlslShaderFrequency Frequency, _mesa_glsl_parse_state* ParseState,
+	static ir_rvalue* GenerateInputFromSemantic(EHlslShaderFrequency Frequency, bool bIsDesktop, _mesa_glsl_parse_state* ParseState,
 		const char* Semantic, const glsl_type* Type, exec_list* DeclInstructions)
 	{
 		if (!Semantic)
@@ -300,7 +353,7 @@ namespace MetalUtils
 
 		if (FCStringAnsi::Strnicmp(Semantic, "SV_", 3) == 0)
 		{
-			FSystemValue* SystemValues = SystemValueTable[Frequency];
+			FSystemValue* SystemValues = bIsDesktop ? DesktopSystemValueTable[Frequency] : MobileSystemValueTable[Frequency];
 			for (int i = 0; SystemValues[i].HlslSemantic != NULL; ++i)
 			{
 				if (SystemValues[i].Mode == ir_var_in && FCStringAnsi::Stricmp(SystemValues[i].HlslSemantic, Semantic) == 0)
@@ -344,9 +397,17 @@ namespace MetalUtils
 		}
 
 		// If we're here, no built-in variables matched.
+		bool bUseSlice = false;
 		if (FCStringAnsi::Strnicmp(Semantic, "SV_", 3) == 0)
 		{
-			_mesa_glsl_warning(ParseState, "unrecognized system value input '%s'", Semantic);
+			if (FCStringAnsi::Strnicmp(Semantic, "SV_RenderTargetArrayIndex", 25) == 0)
+			{
+				bUseSlice = true;
+			}
+			else
+			{
+				_mesa_glsl_warning(ParseState, "unrecognized system value input '%s'", Semantic);
+			}
 		}
 
 		ir_variable* Variable = new(ParseState)ir_variable(
@@ -364,6 +425,11 @@ namespace MetalUtils
 				_mesa_glsl_warning(ParseState, "Unrecognized input attribute '%s'", Semantic);
 			}
 		}
+		else if (bUseSlice)
+		{
+			check(Frequency == HSF_PixelShader);
+			Semantic = NV_EPIC_LAYER;
+		}
 
 		if (!Variable->semantic)
 		{
@@ -376,7 +442,7 @@ namespace MetalUtils
 		return VariableDeref;
 	}
 
-	static void GenerateInputForVariable(EHlslShaderFrequency Frequency, _mesa_glsl_parse_state* ParseState,
+	static void GenerateInputForVariable(EHlslShaderFrequency Frequency, bool bIsDesktop, _mesa_glsl_parse_state* ParseState,
 		const char* InputSemantic, ir_dereference* InputVariableDeref, exec_list* DeclInstructions, exec_list* PreCallInstructions)
 	{
 		const glsl_type* InputType = InputVariableDeref->type;
@@ -413,7 +479,7 @@ namespace MetalUtils
 					ir_dereference_record* FieldDeref = new(ParseState)ir_dereference_record(
 						InputVariableDeref->clone(ParseState, NULL),
 						InputType->fields.structure[i].name);
-					GenerateInputForVariable(Frequency, ParseState, Semantic, FieldDeref, DeclInstructions, PreCallInstructions);
+					GenerateInputForVariable(Frequency, bIsDesktop, ParseState, Semantic, FieldDeref, DeclInstructions, PreCallInstructions);
 				}
 				else
 				{
@@ -443,6 +509,7 @@ namespace MetalUtils
 						);
 					GenerateInputForVariable(
 						Frequency,
+						bIsDesktop,
 						ParseState,
 						ralloc_asprintf(ParseState, "%s%d", Semantic, BaseIndex + i),
 						ArrayDeref,
@@ -452,7 +519,7 @@ namespace MetalUtils
 			}
 			else
 			{
-				ir_rvalue* SrcValue = GenerateInputFromSemantic(Frequency, ParseState, InputSemantic, InputType, DeclInstructions);
+				ir_rvalue* SrcValue = GenerateInputFromSemantic(Frequency, bIsDesktop, ParseState, InputSemantic, InputType, DeclInstructions);
 				if (SrcValue)
 				{
 					YYLTYPE loc;
@@ -464,23 +531,23 @@ namespace MetalUtils
 		}
 	}
 
-	ir_dereference_variable* GenerateInput(EHlslShaderFrequency Frequency, _mesa_glsl_parse_state* ParseState, const char* InputSemantic, const glsl_type* InputType, exec_list* DeclInstructions, exec_list* PreCallInstructions)
+	ir_dereference_variable* GenerateInput(EHlslShaderFrequency Frequency, bool bIsDesktop, _mesa_glsl_parse_state* ParseState, const char* InputSemantic, const glsl_type* InputType, exec_list* DeclInstructions, exec_list* PreCallInstructions)
 	{
 		ir_variable* TempVariable = new(ParseState)ir_variable(InputType, nullptr, ir_var_temporary);
 		ir_dereference_variable* TempVariableDeref = new(ParseState)ir_dereference_variable(TempVariable);
 		PreCallInstructions->push_tail(TempVariable);
-		GenerateInputForVariable(Frequency, ParseState, InputSemantic, TempVariableDeref, DeclInstructions, PreCallInstructions);
+		GenerateInputForVariable(Frequency, bIsDesktop, ParseState, InputSemantic, TempVariableDeref, DeclInstructions, PreCallInstructions);
 		return TempVariableDeref;
 	}
 
-	static ir_rvalue* GenerateOutputFromSemantic(EHlslShaderFrequency Frequency, _mesa_glsl_parse_state* ParseState,
+	static ir_rvalue* GenerateOutputFromSemantic(EHlslShaderFrequency Frequency, bool bIsDesktop, _mesa_glsl_parse_state* ParseState,
 		const char* Semantic, const glsl_type* Type, exec_list* DeclInstructions, const glsl_type** DestVariableType)
 	{
 		ir_variable* Variable = NULL;
 
 		if (FCStringAnsi::Strnicmp(Semantic, "SV_", 3) == 0)
 		{
-			FSystemValue* SystemValues = SystemValueTable[Frequency];
+			FSystemValue* SystemValues = bIsDesktop ? DesktopSystemValueTable[Frequency] : MobileSystemValueTable[Frequency];
 			for (int i = 0; SystemValues[i].HlslSemantic != nullptr; ++i)
 			{
 				if (SystemValues[i].Mode == ir_var_out && FCStringAnsi::Stricmp(SystemValues[i].HlslSemantic, Semantic) == 0 && SystemValues[i].Type->vector_elements == Type->vector_elements)
@@ -510,7 +577,7 @@ namespace MetalUtils
 		return VariableDeref;
 	}
 
-	static void GenerateOutputForVariable(EHlslShaderFrequency Frequency, _mesa_glsl_parse_state* ParseState,
+	static void GenerateOutputForVariable(EHlslShaderFrequency Frequency, bool bIsDesktop, _mesa_glsl_parse_state* ParseState,
 		const char* OutputSemantic, ir_dereference* OutputVariableDeref, exec_list* DeclInstructions, exec_list* PostCallInstructions
 		/*int SemanticArraySize,int SemanticArrayIndex*/)
 	{
@@ -549,7 +616,7 @@ namespace MetalUtils
 					ir_dereference* FieldDeref = new(ParseState)ir_dereference_record(
 						OutputVariableDeref->clone(ParseState, NULL),
 						OutputType->fields.structure[i].name);
-					GenerateOutputForVariable(Frequency, ParseState, Semantic, FieldDeref, DeclInstructions, PostCallInstructions);
+					GenerateOutputForVariable(Frequency, bIsDesktop, ParseState, Semantic, FieldDeref, DeclInstructions, PostCallInstructions);
 				}
 				else
 				{
@@ -583,7 +650,7 @@ namespace MetalUtils
 							OutputVariableDeref->clone(ParseState, NULL),
 							new(ParseState) ir_constant((unsigned)i)
 							);
-						GenerateOutputForVariable(Frequency, ParseState,
+						GenerateOutputForVariable(Frequency, bIsDesktop, ParseState,
 							ralloc_asprintf(ParseState, "%s%d", Semantic, BaseIndex + i),
 							ArrayDeref, DeclInstructions, PostCallInstructions);
 					}
@@ -593,7 +660,7 @@ namespace MetalUtils
 					YYLTYPE loc;
 					ir_rvalue* Src = OutputVariableDeref->clone(ParseState, NULL);
 					const glsl_type* DestVariableType = NULL;
-					ir_rvalue* DestVariableDeref = GenerateOutputFromSemantic(Frequency, ParseState, OutputSemantic,
+					ir_rvalue* DestVariableDeref = GenerateOutputFromSemantic(Frequency, bIsDesktop, ParseState, OutputSemantic,
 						OutputType, DeclInstructions, &DestVariableType);
 
 					apply_type_conversion(DestVariableType, Src, PostCallInstructions, ParseState, true, &loc);
@@ -603,7 +670,7 @@ namespace MetalUtils
 		}
 	}
 
-	ir_dereference_variable* GenerateOutput(EHlslShaderFrequency Frequency, _mesa_glsl_parse_state* ParseState,
+	ir_dereference_variable* GenerateOutput(EHlslShaderFrequency Frequency, bool bIsDesktop, _mesa_glsl_parse_state* ParseState,
 		const char* OutputSemantic, const glsl_type* OutputType, exec_list* DeclInstructions, exec_list* PreCallInstructions, exec_list* PostCallInstructions)
 	{
 		// Generate a local variable to hold the output.
@@ -611,7 +678,7 @@ namespace MetalUtils
 		ir_dereference_variable* TempVariableDeref = new(ParseState) ir_dereference_variable(TempVariable);
 		PreCallInstructions->push_tail(TempVariable);
 
-		GenerateOutputForVariable(Frequency, ParseState, OutputSemantic, TempVariableDeref, DeclInstructions, PostCallInstructions);
+		GenerateOutputForVariable(Frequency, bIsDesktop, ParseState, OutputSemantic, TempVariableDeref, DeclInstructions, PostCallInstructions);
 
 		return TempVariableDeref;
 	}
@@ -644,6 +711,27 @@ struct FFixIntrinsicsVisitor : public ir_rvalue_visitor
 		{
 			return;
 		}
+
+
+		// Fix .x swizzle of scalars...
+		auto* Swizzle = (*RValue)->as_swizzle();
+		if (Swizzle)
+		{
+			auto* Texture = Swizzle->val->as_texture();
+			if (Texture && Texture->op == ir_txf && 
+				Texture->sampler && Texture->sampler->type->sampler_buffer && 
+				Texture->sampler->type->inner_type && Texture->sampler->type->inner_type->is_scalar() && 
+				Swizzle->mask.x == 0 &&
+				Swizzle->mask.y == 0 &&
+				Swizzle->mask.z == 0 &&
+				Swizzle->mask.w == 0 &&
+				Swizzle->mask.num_components == 1 &&
+				Swizzle->mask.has_duplicates == 0)
+			{
+				*RValue = Texture;
+			}
+		}
+
 		auto* expr = (*RValue)->as_expression();
 		if (!expr)
 		{
@@ -873,7 +961,7 @@ void FMetalCodeBackend::MovePackedUniformsToMain(exec_list* ir, _mesa_glsl_parse
 		auto* Var = Instruction->as_variable();
 		if (Var)
 		{
-			check(Var->mode == ir_var_uniform || Var->mode == ir_var_out || Var->mode == ir_var_in);
+			check(Var->mode == ir_var_uniform || Var->mode == ir_var_out || Var->mode == ir_var_in || Var->mode == ir_var_shared);
 			OutBuffers.Buffers.Add(Var);
 		}
 	}
@@ -952,10 +1040,10 @@ void FMetalCodeBackend::PromoteInputsAndOutputsGlobalHalfToFloat(exec_list* Inst
 	//IRDump(Instructions);
 }
 
-static bool ProcessStageInVariables(_mesa_glsl_parse_state* ParseState, EHlslShaderFrequency Frequency, ir_variable* Variable, TArray<glsl_struct_field>& OutStageInMembers, std::set<ir_variable*>& OutStageInVariables, unsigned int* OutVertexAttributesMask, TIRVarList& OutFunctionArguments)
+static bool ProcessStageInVariables(_mesa_glsl_parse_state* ParseState, bool bIsDesktop, EHlslShaderFrequency Frequency, ir_variable* Variable, TArray<glsl_struct_field>& OutStageInMembers, std::set<ir_variable*>& OutStageInVariables, unsigned int* OutVertexAttributesMask, TIRVarList& OutFunctionArguments)
 {
 	// Don't move variables that are system values into the input structures
-	const auto* SystemValues = MetalUtils::SystemValueTable[Frequency];
+	const auto* SystemValues = bIsDesktop ? MetalUtils::DesktopSystemValueTable[Frequency] : MetalUtils::MobileSystemValueTable[Frequency];
 	for(int i = 0; SystemValues[i].MetalSemantic != nullptr; ++i)
 	{
 		if (SystemValues[i].Mode == ir_var_in && FCStringAnsi::Stricmp(SystemValues[i].MetalSemantic, Variable->semantic) == 0)
@@ -1149,7 +1237,8 @@ static FSystemValue VertexSystemValueTable[] =
 {
 	{"SV_VertexID", glsl_type::uint_type, "gl_VertexID", ir_var_in, false, false},
 	{"SV_InstanceID", glsl_type::uint_type, "gl_InstanceID", ir_var_in, false, false},
-	//{ "SV_Position", glsl_type::vec4_type, "gl_Position", ir_var_out, false, true },
+    {"SV_RenderTargetArrayIndex", glsl_type::uint_type, "OUT_Layer", ir_var_out, "SV_RenderTargetArrayIndex"},
+    //{ "SV_Position", glsl_type::vec4_type, "gl_Position", ir_var_out, false, true },
 	{NULL, NULL, NULL, ir_var_auto, false, false}
 };
 
@@ -1159,8 +1248,9 @@ static FSystemValue PixelSystemValueTable[] =
 	{"SV_Depth", glsl_type::float_type, "gl_FragDepth", ir_var_out, false, false},
 	{"SV_Position", glsl_type::vec4_type, "gl_FragCoord", ir_var_in, true, false},
 	//	{ "SV_IsFrontFace", glsl_type::bool_type, "gl_FrontFacing", ir_var_in, false, true },
-	{"SV_PrimitiveID", glsl_type::int_type, "gl_PrimitiveID", ir_var_in, false, false},
-	//	{ "SV_RenderTargetArrayIndex", glsl_type::int_type, "gl_Layer", ir_var_in, false, false },
+    {"SV_PrimitiveID", glsl_type::int_type, "gl_PrimitiveID", ir_var_in, false, false},
+    {"SV_RenderTargetArrayIndex", glsl_type::uint_type, "IN_Layer", ir_var_in, "SV_RenderTargetArrayIndex"},
+	//	{ "SV_RenderTargetArrayIndex", glsl_type::uint_type, "gl_Layer", ir_var_in, false, false },
 	//	{ "SV_Target0", glsl_type::vec4_type, "gl_FragColor", ir_var_out, false, true },
 	{NULL, NULL, NULL, ir_var_auto, false, false}
 };
@@ -1188,7 +1278,7 @@ static FSystemValue* SystemValueTable[] =
 * @returns the IR variable deref for the semantic.
 */
 static ir_dereference_variable* GenerateShaderInput(
-	EHlslShaderFrequency Frequency,
+	EHlslShaderFrequency Frequency, bool bIsDesktop,
 	_mesa_glsl_parse_state* ParseState,
 	const char* InputSemantic,
 	const glsl_type* InputType,
@@ -1203,7 +1293,7 @@ static ir_dereference_variable* GenerateShaderInput(
 	PreCallInstructions->push_tail(TempVariable);
 
 	check(!InputType->is_inputpatch() && !InputType->is_outputpatch());
-	ir_rvalue* SrcValue = MetalUtils::GenerateInputFromSemantic(Frequency, ParseState, InputSemantic, InputType, DeclInstructions);
+	ir_rvalue* SrcValue = MetalUtils::GenerateInputFromSemantic(Frequency, bIsDesktop, ParseState, InputSemantic, InputType, DeclInstructions);
 	if(SrcValue)
 	{
 		YYLTYPE loc ={0};
@@ -1489,7 +1579,7 @@ void FMetalCodeBackend::PackInputsAndOutputs(exec_list* Instructions, _mesa_glsl
 					break;
 
 				case ir_var_in:
-					if (!ProcessStageInVariables(ParseState, Frequency, Variable, VSStageInMembers, VSStageInVariables, nullptr, VSInputArguments))
+					if (!ProcessStageInVariables(ParseState, bIsDesktop, Frequency, Variable, VSStageInMembers, VSStageInVariables, nullptr, VSInputArguments))
 					{
 						return;
 					}
@@ -1603,7 +1693,7 @@ void FMetalCodeBackend::PackInputsAndOutputs(exec_list* Instructions, _mesa_glsl
 					break;
 
 				case ir_var_in:
-					if (!ProcessStageInVariables(ParseState, Frequency, Variable, PSStageInMembers, PSStageInVariables, nullptr, PSInputArguments))
+					if (!ProcessStageInVariables(ParseState, bIsDesktop, Frequency, Variable, PSStageInMembers, PSStageInVariables, nullptr, PSInputArguments))
 					{
 						return;
 					}
@@ -1667,7 +1757,7 @@ void FMetalCodeBackend::PackInputsAndOutputs(exec_list* Instructions, _mesa_glsl
 					{
 						TArray<glsl_struct_field> CSStageInMembers;
 						TIRVarSet CSStageInVariables;
-						if (!ProcessStageInVariables(ParseState, Frequency, Variable, CSStageInMembers, CSStageInVariables, nullptr, CSInputArguments))
+						if (!ProcessStageInVariables(ParseState, bIsDesktop, Frequency, Variable, CSStageInMembers, CSStageInVariables, nullptr, CSInputArguments))
 						{
 							return;
 						}
@@ -1749,7 +1839,7 @@ void FMetalCodeBackend::PackInputsAndOutputs(exec_list* Instructions, _mesa_glsl
 					// At this point this should be a built-in system value
 					check(Variable->semantic);
 					ArgVarDeref = GenerateShaderInput(
-						Frequency,
+						Frequency, bIsDesktop,
 						ParseState,
 						Variable->semantic,
 						Variable->type,
diff --git a/Engine/Source/Developer/iOS/MetalShaderFormat/Private/MetalUtils.h b/Engine/Source/Developer/iOS/MetalShaderFormat/Private/MetalUtils.h
--- a/Engine/Source/Developer/iOS/MetalShaderFormat/Private/MetalUtils.h
+++ b/Engine/Source/Developer/iOS/MetalShaderFormat/Private/MetalUtils.h
@@ -38,14 +38,17 @@ struct extern_var : public exec_node
 
 struct FBuffers
 {
+	TIRVarSet AtomicVariables;
+
 	TArray<class ir_instruction*> Buffers;
 
 	// Information about textures & samplers; we need to have unique samplerstate indices, as one they can be used independent of each other
 	TArray<std::string> UniqueSamplerStates;
 
-	int32 GetUniqueSamplerStateIndex(const std::string& Name, bool bAddIfNotFound)
+	int32 GetUniqueSamplerStateIndex(const std::string& Name, bool bAddIfNotFound, bool& bOutAdded)
 	{
 		int32 Found = INDEX_NONE;
+		bOutAdded = false;
 		if (UniqueSamplerStates.Find(Name, Found))
 		{
 			return Found;
@@ -54,18 +57,24 @@ struct FBuffers
 		if (bAddIfNotFound)
 		{
 			UniqueSamplerStates.Add(Name);
+			bOutAdded = true;
 			return UniqueSamplerStates.Num() - 1;
 		}
 
 		return Found;
 	}
 
-	int GetIndex(ir_variable* Var)
+	int GetIndex(ir_variable* Var, bool bIsDesktop)
 	{
 		for (int i = 0, n = Buffers.Num(); i < n; ++i)
 		{
 			if (Buffers[i] == Var)
 			{
+				if (Var->type->is_image())
+				{
+					//return (bIsDesktop ? 126 : 30) - i;
+					return 30 - i;
+				}
 				return i;
 			}
 		}
@@ -73,13 +82,18 @@ struct FBuffers
 		return -1;
 	}
 
-	int GetIndex(const std::string& Name)
+	int GetIndex(const std::string& Name, bool bIsDesktop)
 	{
 		for (int i = 0, n = Buffers.Num(); i < n; ++i)
 		{
 			auto* Var = Buffers[i] ? Buffers[i]->as_variable() : nullptr;
 			if (Var && Var->name && Var->name == Name)
 			{
+				if (Var->type->is_image())
+				{
+					//return (bIsDesktop ? 126 : 30) - i;
+					return 30 - i;
+				}
 				return i;
 			}
 		}
@@ -133,11 +147,13 @@ const glsl_type* PromoteHalfToFloatType(_mesa_glsl_parse_state* state, const gls
 
 namespace MetalUtils
 {
-	ir_dereference_variable* GenerateInput(EHlslShaderFrequency Frequency, _mesa_glsl_parse_state* ParseState, const char* InputSemantic,
+	ir_dereference_variable* GenerateInput(EHlslShaderFrequency Frequency, bool bIsDesktop, _mesa_glsl_parse_state* ParseState, const char* InputSemantic,
 		const glsl_type* InputType, exec_list* DeclInstructions, exec_list* PreCallInstructions);
 
-	ir_dereference_variable* GenerateOutput(EHlslShaderFrequency Frequency, _mesa_glsl_parse_state* ParseState, const char* OutputSemantic,
+	ir_dereference_variable* GenerateOutput(EHlslShaderFrequency Frequency, bool bIsDesktop, _mesa_glsl_parse_state* ParseState, const char* OutputSemantic,
 		const glsl_type* OutputType, exec_list* DeclInstructions, exec_list* PreCallInstructions, exec_list* PostCallInstructions);
 }
 
 const int MAX_SIMULTANEOUS_RENDER_TARGETS = 8;
+
+#define NV_EPIC_LAYER							"__epic_layer_index"
