Metal: Fixed sprintf warning

diff --git a/Engine/Source/Developer/iOS/MetalShaderFormat/Private/MetalShaderCompiler.cpp b/Engine/Source/Developer/iOS/MetalShaderFormat/Private/MetalShaderCompiler.cpp
--- a/Engine/Source/Developer/iOS/MetalShaderFormat/Private/MetalShaderCompiler.cpp
+++ b/Engine/Source/Developer/iOS/MetalShaderFormat/Private/MetalShaderCompiler.cpp
@@ -1022,7 +1022,7 @@ void CompileShader_Metal(const FShaderCompilerInput& Input,FShaderCompilerOutput
 		{
 			uint32 Len = FCStringAnsi::Strlen(TCHAR_TO_ANSI(*Input.SourceFilename)) + FCStringAnsi::Strlen(TCHAR_TO_ANSI(*Input.EntryPointName)) + FCStringAnsi::Strlen(MetalShaderSource) + 20;
 			char* Dest = (char*)malloc(Len);
-			sprintf(Dest, "// ! %s.usf:%s\n%s", TCHAR_TO_ANSI(*Input.SourceFilename), TCHAR_TO_ANSI(*Input.EntryPointName), MetalShaderSource);
+			sprintf_s(Dest, Len, "// ! %s.usf:%s\n%s", TCHAR_TO_ANSI(*Input.SourceFilename), TCHAR_TO_ANSI(*Input.EntryPointName), MetalShaderSource);
 			free(MetalShaderSource);
 			MetalShaderSource = Dest;
 			SourceLen = FCStringAnsi::Strlen(MetalShaderSource);
