Buildwathcher - Fix initialization order

diff --git a/Engine/Source/Developer/iOS/MetalShaderFormat/Private/MetalBackend.cpp b/Engine/Source/Developer/iOS/MetalShaderFormat/Private/MetalBackend.cpp
--- a/Engine/Source/Developer/iOS/MetalShaderFormat/Private/MetalBackend.cpp
+++ b/Engine/Source/Developer/iOS/MetalShaderFormat/Private/MetalBackend.cpp
@@ -2801,10 +2801,10 @@ struct FMetalCheckRestrictionsVisitor : public ir_hierarchical_visitor
 	_mesa_glsl_parse_state* ParseState;
 	EHlslShaderFrequency Frequency;
 	bool bErrors;
-	FMetalCheckRestrictionsVisitor(_mesa_glsl_parse_state* InParseState, EHlslShaderFrequency InFrequency) :
-		ParseState(InParseState),
-		bErrors(false),
-		Frequency(InFrequency)
+	FMetalCheckRestrictionsVisitor(_mesa_glsl_parse_state* InParseState, EHlslShaderFrequency InFrequency)
+		: ParseState(InParseState),
+		, Frequency(InFrequency)
+		, bErrors(false)
 	{
 	}
 
