Remove extra ,

diff --git a/Engine/Source/Developer/iOS/MetalShaderFormat/Private/MetalBackend.cpp b/Engine/Source/Developer/iOS/MetalShaderFormat/Private/MetalBackend.cpp
--- a/Engine/Source/Developer/iOS/MetalShaderFormat/Private/MetalBackend.cpp
+++ b/Engine/Source/Developer/iOS/MetalShaderFormat/Private/MetalBackend.cpp
@@ -2802,7 +2802,7 @@ struct FMetalCheckRestrictionsVisitor : public ir_hierarchical_visitor
 	EHlslShaderFrequency Frequency;
 	bool bErrors;
 	FMetalCheckRestrictionsVisitor(_mesa_glsl_parse_state* InParseState, EHlslShaderFrequency InFrequency)
-		: ParseState(InParseState),
+		: ParseState(InParseState)
 		, Frequency(InFrequency)
 		, bErrors(false)
 	{
