Fixed problems with compiling Metal shaders on pre-10.11 versions of OS X

diff --git a/Engine/Shaders/Common.usf b/Engine/Shaders/Common.usf
--- a/Engine/Shaders/Common.usf
+++ b/Engine/Shaders/Common.usf
@@ -11,7 +11,7 @@
 	#include "PS4/PS4Common.usf"
 #endif
 
-#if METAL_PROFILE || METAL_MRT_PROFILE
+#if METAL_PROFILE || METAL_MRT_PROFILE || METAL_SM5_PROFILE
 	// Helps with iteration when changing Metal shader code generation backend.
 	#include "MetalCommon.usf"
 #endif
@@ -62,7 +62,7 @@
 	#define MaterialFloat4x3 float4x3 
 #endif
 
-#if COMPILER_GLSL_ES2 || METAL_PROFILE || METAL_MRT_PROFILE
+#if COMPILER_GLSL_ES2 || METAL_PROFILE || METAL_MRT_PROFILE || METAL_SM5_PROFILE
 
 	float determinant(float3x3 M)
 	{    
diff --git a/Engine/Source/Developer/iOS/MetalShaderFormat/Private/MetalShaderCompiler.cpp b/Engine/Source/Developer/iOS/MetalShaderFormat/Private/MetalShaderCompiler.cpp
--- a/Engine/Source/Developer/iOS/MetalShaderFormat/Private/MetalShaderCompiler.cpp
+++ b/Engine/Source/Developer/iOS/MetalShaderFormat/Private/MetalShaderCompiler.cpp
@@ -714,28 +714,25 @@ static void BuildMetalShaderOutput(
 		bool bSucceeded = false;
 
 #if PLATFORM_MAC
-		// metal commandlines
-		FString Standard = (bIsMobile ? TEXT("-std=ios-metal1.0") : TEXT("-std=osx-metal1.1 -Wno-null-character -ffmast-math"));
-		
-		FString Params = FString::Printf(TEXT("%s %s -o %s"), *Standard, *InputFilename, *ObjFilename);
-		FPlatformProcess::ExecProcess( TEXT("/Applications/Xcode.app/Contents/Developer/Toolchains/XcodeDefault.xctoolchain/usr/bin/metal"), *Params, &ReturnCode, &Results, &Errors );
-
-		// handle compile error
-		if (ReturnCode != 0 || IFileManager::Get().FileSize(*ObjFilename) <= 0)
-		{
-			FShaderCompilerError* Error = new(OutErrors) FShaderCompilerError();
-			Error->ErrorFile = InputFilename;
-			Error->ErrorLineString = TEXT("0");
-			Error->StrippedErrorMessage = Results + Errors;
-			bSucceeded = false;
-		}
-		else
+		FString XcodePath = FPlatformMisc::GetXcodePath();
+		if (XcodePath.Len() > 0 && (bIsMobile || FPlatformMisc::MacOSXVersionCompare(10, 11, 0) >= 0))
 		{
-			Params = FString::Printf(TEXT("r %s %s"), *ArFilename, *ObjFilename);
-			FPlatformProcess::ExecProcess( TEXT("/Applications/Xcode.app/Contents/Developer/Toolchains/XcodeDefault.xctoolchain/usr/bin/metal-ar"), *Params, &ReturnCode, &Results, &Errors );
+			FString MetalToolsPath = FString::Printf(TEXT("%s/Toolchains/XcodeDefault.xctoolchain/usr/bin"), *XcodePath);
+			FString MetalPath = MetalToolsPath + TEXT("/metal");
+			if (IFileManager::Get().FileSize(*MetalPath) <= 0)
+			{
+				MetalToolsPath = FString::Printf(TEXT("%s/Platforms/iPhoneOS.platform/usr/bin"), *XcodePath);
+				MetalPath = MetalToolsPath + TEXT("/metal");
+			}
+
+			// metal commandlines
+			FString Standard = (bIsMobile ? TEXT("-std=ios-metal1.0") : TEXT("-std=osx-metal1.1 -Wno-null-character -ffmast-math"));
+
+			FString Params = FString::Printf(TEXT("%s %s -o %s"), *Standard, *InputFilename, *ObjFilename);
+			FPlatformProcess::ExecProcess( *MetalPath, *Params, &ReturnCode, &Results, &Errors );
 
 			// handle compile error
-			if (ReturnCode != 0 || IFileManager::Get().FileSize(*ArFilename) <= 0)
+			if (ReturnCode != 0 || IFileManager::Get().FileSize(*ObjFilename) <= 0)
 			{
 				FShaderCompilerError* Error = new(OutErrors) FShaderCompilerError();
 				Error->ErrorFile = InputFilename;
@@ -745,11 +742,12 @@ static void BuildMetalShaderOutput(
 			}
 			else
 			{
-				Params = FString::Printf(TEXT("-o %s %s"), *OutputFilename, *ArFilename);
-				FPlatformProcess::ExecProcess( TEXT("/Applications/Xcode.app/Contents/Developer/Toolchains/XcodeDefault.xctoolchain/usr/bin/metallib"), *Params, &ReturnCode, &Results, &Errors );
-		
+				Params = FString::Printf(TEXT("r %s %s"), *ArFilename, *ObjFilename);
+				FString MetalArPath = MetalToolsPath + TEXT("/metal-ar");
+				FPlatformProcess::ExecProcess( *MetalArPath, *Params, &ReturnCode, &Results, &Errors );
+
 				// handle compile error
-				if (ReturnCode != 0 || IFileManager::Get().FileSize(*OutputFilename) <= 0)
+				if (ReturnCode != 0 || IFileManager::Get().FileSize(*ArFilename) <= 0)
 				{
 					FShaderCompilerError* Error = new(OutErrors) FShaderCompilerError();
 					Error->ErrorFile = InputFilename;
@@ -759,25 +757,41 @@ static void BuildMetalShaderOutput(
 				}
 				else
 				{
-					bCompileAtRuntime = false;
-					
-					// Write out the header and compiled shader code
-					FMemoryWriter Ar(ShaderOutput.Code, true);
-					uint8 PrecompiledFlag = 1;
-					Ar << PrecompiledFlag;
-					Ar << Header;
-
-					// load output
-					TArray<uint8> CompiledShader;
-					FFileHelper::LoadFileToArray(CompiledShader, *OutputFilename);
-					
-					// jam it into the output bytes
-					Ar.Serialize(CompiledShader.GetData(), CompiledShader.Num());
-					
-					ShaderOutput.NumInstructions = 0;
-					ShaderOutput.NumTextureSamplers = Header.Bindings.NumSamplers;
-					ShaderOutput.bSucceeded = true;
-					bSucceeded = true;
+					Params = FString::Printf(TEXT("-o %s %s"), *OutputFilename, *ArFilename);
+					FString MetalLibPath = MetalToolsPath + TEXT("/metallib");
+					FPlatformProcess::ExecProcess( *MetalLibPath, *Params, &ReturnCode, &Results, &Errors );
+			
+					// handle compile error
+					if (ReturnCode != 0 || IFileManager::Get().FileSize(*OutputFilename) <= 0)
+					{
+						FShaderCompilerError* Error = new(OutErrors) FShaderCompilerError();
+						Error->ErrorFile = InputFilename;
+						Error->ErrorLineString = TEXT("0");
+						Error->StrippedErrorMessage = Results + Errors;
+						bSucceeded = false;
+					}
+					else
+					{
+						bCompileAtRuntime = false;
+						
+						// Write out the header and compiled shader code
+						FMemoryWriter Ar(ShaderOutput.Code, true);
+						uint8 PrecompiledFlag = 1;
+						Ar << PrecompiledFlag;
+						Ar << Header;
+
+						// load output
+						TArray<uint8> CompiledShader;
+						FFileHelper::LoadFileToArray(CompiledShader, *OutputFilename);
+						
+						// jam it into the output bytes
+						Ar.Serialize(CompiledShader.GetData(), CompiledShader.Num());
+						
+						ShaderOutput.NumInstructions = 0;
+						ShaderOutput.NumTextureSamplers = Header.Bindings.NumSamplers;
+						ShaderOutput.bSucceeded = true;
+						bSucceeded = true;
+					}
 				}
 			}
 		}
@@ -887,7 +901,7 @@ void CompileShader_Metal(const FShaderCompilerInput& Input,FShaderCompilerOutput
 
 	// @todo - Zebra - Work out which standard we need, this is dependent on the shader platform.
 	// For now only SP_METAL will compile for iOS, with MRT & SM5 only for Mac.
-	const bool bIsMobile = (Output.Target.Platform == SP_METAL);
+	const bool bIsMobile = (Input.Target.Platform == SP_METAL);
 	if (bIsMobile)
 	{
 		AdditionalDefines.SetDefine(TEXT("IOS"), 1);
diff --git a/Engine/Source/Runtime/Core/Private/Mac/MacPlatformMisc.cpp b/Engine/Source/Runtime/Core/Private/Mac/MacPlatformMisc.cpp
--- a/Engine/Source/Runtime/Core/Private/Mac/MacPlatformMisc.cpp
+++ b/Engine/Source/Runtime/Core/Private/Mac/MacPlatformMisc.cpp
@@ -72,9 +72,15 @@ struct FMacApplicationInfo
 			NSDictionary* SystemVersion = [NSDictionary dictionaryWithContentsOfFile: @"/System/Library/CoreServices/SystemVersion.plist"];
 			OSBuild = FString((NSString*)[SystemVersion objectForKey: @"ProductBuildVersion"]);
 		}
-		
+
 		RunningOnMavericks = OSXVersion.majorVersion == 10 && OSXVersion.minorVersion == 9;
-		
+
+		FPlatformProcess::ExecProcess(TEXT("/usr/bin/xcode-select"), TEXT("--print-path"), nullptr, &XcodePath, nullptr);
+		if (XcodePath.Len() > 0)
+		{
+			XcodePath.RemoveAt(XcodePath.Len() - 1); // Remove \n at the end of the string
+		}
+
 		char TempSysCtlBuffer[PATH_MAX] = {};
 		size_t TempSysCtlBufferSize = PATH_MAX;
 		
@@ -272,6 +278,7 @@ struct FMacApplicationInfo
 	FString ExecutableName;
 	NSOperatingSystemVersion OSXVersion;
 	FGuid RunUUID;
+	FString XcodePath;
 	static PLCrashReporter* CrashReporter;
 	static FMacMallocCrashHandler* CrashMalloc;
 };
@@ -1234,6 +1241,11 @@ FString FMacPlatformMisc::GetOperatingSystemId()
 	return Result;
 }
 
+FString FMacPlatformMisc::GetXcodePath()
+{
+	return GMacAppInfo.XcodePath;
+}
+
 /** Global pointer to crash handler */
 void (* GCrashHandlerPointer)(const FGenericCrashContext& Context) = NULL;
 
diff --git a/Engine/Source/Runtime/Core/Public/Mac/MacPlatformMisc.h b/Engine/Source/Runtime/Core/Public/Mac/MacPlatformMisc.h
--- a/Engine/Source/Runtime/Core/Public/Mac/MacPlatformMisc.h
+++ b/Engine/Source/Runtime/Core/Public/Mac/MacPlatformMisc.h
@@ -198,6 +198,8 @@ struct CORE_API FMacPlatformMisc : public FGenericPlatformMisc
 	 */
 	static FString GetOperatingSystemId();
 
+	static FString GetXcodePath();
+
 	static bool bChachedMacMenuStateNeedsUpdate;
 
 	static id<NSObject> CommandletActivity;
diff --git a/Engine/Source/Runtime/RHI/Public/RHIDefinitions.h b/Engine/Source/Runtime/RHI/Public/RHIDefinitions.h
--- a/Engine/Source/Runtime/RHI/Public/RHIDefinitions.h
+++ b/Engine/Source/Runtime/RHI/Public/RHIDefinitions.h
@@ -601,7 +601,7 @@ enum class ESimpleRenderTargetMode
 
 inline bool IsPCPlatform(const EShaderPlatform Platform)
 {
-	return Platform == SP_PCD3D_SM5 || Platform == SP_PCD3D_SM4 || Platform == SP_PCD3D_ES2 || Platform == SP_PCD3D_ES3_1 || Platform ==  SP_OPENGL_SM4 || Platform == SP_OPENGL_SM4_MAC || Platform == SP_OPENGL_SM5 || Platform == SP_OPENGL_PCES2 || Platform == SP_OPENGL_PCES3_1 || (PLATFORM_MAC && (Platform == SP_METAL || Platform == SP_METAL_MRT || Platform == SP_METAL_SM5));
+	return Platform == SP_PCD3D_SM5 || Platform == SP_PCD3D_SM4 || Platform == SP_PCD3D_ES2 || Platform == SP_PCD3D_ES3_1 || Platform ==  SP_OPENGL_SM4 || Platform == SP_OPENGL_SM4_MAC || Platform == SP_OPENGL_SM5 || Platform == SP_OPENGL_PCES2 || Platform == SP_OPENGL_PCES3_1 || (PLATFORM_MAC && (Platform == SP_METAL_MRT || Platform == SP_METAL_SM5));
 }
 
 /** Whether the shader platform corresponds to the ES2 feature level. */
