UE4 - Fix Metal matrix constructor

diff --git a/Engine/Source/Developer/iOS/MetalShaderFormat/Private/MetalBackend.cpp b/Engine/Source/Developer/iOS/MetalShaderFormat/Private/MetalBackend.cpp
--- a/Engine/Source/Developer/iOS/MetalShaderFormat/Private/MetalBackend.cpp
+++ b/Engine/Source/Developer/iOS/MetalShaderFormat/Private/MetalBackend.cpp
@@ -1724,16 +1724,48 @@ protected:
 		}
 		else
 		{
-			print_type_full(constant->type);
-			ralloc_asprintf_append(buffer, "(");
-			print_constant(constant, 0);
-			int num_components = constant->type->components();
-			for (int i = 1; i < num_components; ++i)
+			if (constant->type->is_matrix())
 			{
-				ralloc_asprintf_append(buffer, ",");
-				print_constant(constant, i);
+				// Need to print row by row
+				print_type_full(constant->type);
+				ralloc_asprintf_append(buffer, "(");
+				const auto* RowType = constant->type->column_type();
+				uint32 Component = 0;
+				for (uint32 Index = 0; Index < constant->type->matrix_columns; ++Index)
+				{
+					if (Index > 0)
+					{
+						ralloc_asprintf_append(buffer, ",");
+					}
+					print_type_full(RowType);
+					ralloc_asprintf_append(buffer, "(");
+					for (uint32 VecIndex = 0; VecIndex < RowType->vector_elements; ++VecIndex)
+					{
+						if (VecIndex > 0)
+						{
+							ralloc_asprintf_append(buffer, ",");
+						}
+						print_constant(constant, Component);
+						++Component;
+					}
+					ralloc_asprintf_append(buffer, ")");
+				}
+				check(Component == constant->type->components());
+				ralloc_asprintf_append(buffer, ")");
+			}
+			else
+			{
+				print_type_full(constant->type);
+				ralloc_asprintf_append(buffer, "(");
+				print_constant(constant, 0);
+				int num_components = constant->type->components();
+				for (int i = 1; i < num_components; ++i)
+				{
+					ralloc_asprintf_append(buffer, ",");
+					print_constant(constant, i);
+				}
+				ralloc_asprintf_append(buffer, ")");
 			}
-			ralloc_asprintf_append(buffer, ")");
 		}
 	}
 
@@ -2790,7 +2822,7 @@ char* FMetalCodeBackend::GenerateCode(exec_list* ir, _mesa_glsl_parse_state* sta
 
 		Validate(ir, state);
 	}
-
+//IRDump(ir);
 	// Generate the actual code string
 	const char* code = visitor.run(ir);
 	return _strdup(code);
