UE4 - Fix env var; added COMPILER_HLSLCC to shaders

diff --git a/Engine/Shaders/Definitions.usf b/Engine/Shaders/Definitions.usf
--- a/Engine/Shaders/Definitions.usf
+++ b/Engine/Shaders/Definitions.usf
@@ -58,8 +58,8 @@
 #define OUTPUT_DEPTH_TO_ALPHA							0
 #endif
 
-#ifndef COMPILER_CG
-#define COMPILER_CG										0
+#ifndef COMPILER_HLSLCC
+#define COMPILER_HLSLCC									0
 #endif
 
 #ifndef COMPILER_HLSL
@@ -210,3 +210,10 @@
 #ifndef HAS_INVERTED_Z_BUFFER
 #define HAS_INVERTED_Z_BUFFER	1
 #endif
+
+// 'static' asserts
+#if COMPILER_GLSL || COMPILER_GLSL_ES2 || VULKAN_PROFILE || METAL_PROFILE || METAL_MRT_PROFILE || METAL_SM5_PROFILE
+#if !COMPILER_HLSLCC
+#error "Missing COMPILER_HLSLCC define!"
+#endif
+#endif
diff --git a/Engine/Source/Developer/ShaderFormatOpenGL/Private/OpenGLShaderCompiler.cpp b/Engine/Source/Developer/ShaderFormatOpenGL/Private/OpenGLShaderCompiler.cpp
--- a/Engine/Source/Developer/ShaderFormatOpenGL/Private/OpenGLShaderCompiler.cpp
+++ b/Engine/Source/Developer/ShaderFormatOpenGL/Private/OpenGLShaderCompiler.cpp
@@ -1169,6 +1169,8 @@ void CompileShader_Windows_OGL(const FShaderCompilerInput& Input,FShaderCompiler
 	FString PreprocessedShader;
 	FShaderCompilerDefinitions AdditionalDefines;
 	EHlslCompileTarget HlslCompilerTarget = HCT_InvalidTarget;
+
+	AdditionalDefines.SetDefine(TEXT("COMPILER_HLSLCC"), 1);
 	switch (Version)
 	{
 		case GLSL_310_ES_EXT:
diff --git a/Engine/Source/Developer/iOS/MetalShaderFormat/Private/MetalShaderCompiler.cpp b/Engine/Source/Developer/iOS/MetalShaderFormat/Private/MetalShaderCompiler.cpp
--- a/Engine/Source/Developer/iOS/MetalShaderFormat/Private/MetalShaderCompiler.cpp
+++ b/Engine/Source/Developer/iOS/MetalShaderFormat/Private/MetalShaderCompiler.cpp
@@ -465,6 +465,8 @@ void CompileShader_Metal(const FShaderCompilerInput& Input,FShaderCompilerOutput
 	FShaderCompilerDefinitions AdditionalDefines;
 	EHlslCompileTarget HlslCompilerTarget = HCT_FeatureLevelES3_1;
 
+	AdditionalDefines.SetDefine(TEXT("COMPILER_HLSLCC"), 1 );
+
 	// @todo - Zebra - Work out which standard we need, this is dependent on the shader platform.
 	// For now only SP_METAL will compile for iOS, with MRT & SM5 only for Mac.
 	const bool bIsMobile = (Input.Target.Platform == SP_METAL);
