re-enable offline shader compiling for Metal on Mac

diff --git a/Engine/Source/Developer/iOS/MetalShaderFormat/Private/MetalShaderFormat.h b/Engine/Source/Developer/iOS/MetalShaderFormat/Private/MetalShaderFormat.h
--- a/Engine/Source/Developer/iOS/MetalShaderFormat/Private/MetalShaderFormat.h
+++ b/Engine/Source/Developer/iOS/MetalShaderFormat/Private/MetalShaderFormat.h
@@ -6,5 +6,5 @@ extern void CompileShader_Metal(const struct FShaderCompilerInput& Input, struct
 
 // Set this to 0 to get shader source in the graphics debugger
 // Note: Offline and runtime compiled shaders have separate DDC versions and can co-exist
-#define METAL_OFFLINE_COMPILE 0 // !PLATFORM_WINDOWS
+#define METAL_OFFLINE_COMPILE !PLATFORM_WINDOWS
 
