Tweaking some stuff in the iOS project settings for localization and filtering support.

diff --git a/Engine/Source/Developer/iOS/IOSPlatformEditor/Private/IOSTargetSettingsCustomization.cpp b/Engine/Source/Developer/iOS/IOSPlatformEditor/Private/IOSTargetSettingsCustomization.cpp
--- a/Engine/Source/Developer/iOS/IOSPlatformEditor/Private/IOSTargetSettingsCustomization.cpp
+++ b/Engine/Source/Developer/iOS/IOSPlatformEditor/Private/IOSTargetSettingsCustomization.cpp
@@ -421,7 +421,7 @@ void FIOSTargetSettingsCustomization::BuildPListSection(IDetailLayoutBuilder& De
 						.IsEnabled(this, &FIOSTargetSettingsCustomization::IsImportEnabled)
 						[
 							SNew(STextBlock)
-							.Text(FText::FromString("Import Provision"))
+							.Text(LOCTEXT("ImportProvision", "Import Provision"))
 						]
 					]
 			]
@@ -576,7 +576,7 @@ void FIOSTargetSettingsCustomization::BuildPListSection(IDetailLayoutBuilder& De
 					.IsEnabled(this, &FIOSTargetSettingsCustomization::IsImportEnabled)
 					[
 						SNew(STextBlock)
-						.Text(FText::FromString("Import Certificate"))
+						.Text(LOCTEXT("ImportCertificate", "Import Certificate"))
 					]
 				]
 			]
@@ -803,8 +803,11 @@ void FIOSTargetSettingsCustomization::BuildRemoteBuildingSection(IDetailLayoutBu
 	SSHPrivateKeyOverridePathPropertyRow
 		.ToolTip(LOCTEXT("SSHPrivateKeyOverridePathToolTip", "Override the existing SSH Private Key with one from a specified location."));
 
+	const FText GenerateSSHText = LOCTEXT("GenerateSSHKey", "Generate SSH Key");
+
 	// Add a generate key button
 	RemoteBuildingGroup.AddWidgetRow()
+		.FilterString(GenerateSSHText)
 		.WholeRowWidget
 		.MinDesiredWidth(0.f)
 		.MaxDesiredWidth(0.f)
@@ -826,7 +829,7 @@ void FIOSTargetSettingsCustomization::BuildRemoteBuildingSection(IDetailLayoutBu
 						.IsEnabled(this, &FIOSTargetSettingsCustomization::IsImportEnabled)
 						[
 							SNew(STextBlock)
-							.Text(FText::FromString("Generate SSH Key"))
+							.Text(GenerateSSHText)
 						]
 					]
 				]
