Addition of mobile provision selection
no error messaging as of yet when selecting something which then expires

diff --git a/Engine/Source/Developer/iOS/IOSPlatformEditor/Private/IOSTargetSettingsCustomization.cpp b/Engine/Source/Developer/iOS/IOSPlatformEditor/Private/IOSTargetSettingsCustomization.cpp
--- a/Engine/Source/Developer/iOS/IOSPlatformEditor/Private/IOSTargetSettingsCustomization.cpp
+++ b/Engine/Source/Developer/iOS/IOSPlatformEditor/Private/IOSTargetSettingsCustomization.cpp
@@ -76,6 +76,7 @@ FIOSTargetSettingsCustomization::FIOSTargetSettingsCustomization()
 
 	bShowAllProvisions = false;
 	bShowAllCertificates = false;
+	ProvisionList = MakeShareable(new TArray<ProvisionPtr>());
 }
 
 FIOSTargetSettingsCustomization::~FIOSTargetSettingsCustomization()
@@ -113,7 +114,7 @@ void FIOSTargetSettingsCustomization::UpdateStatus()
 	if (OutputMessage.Len() > 0)
 	{
 		CertificateList.Reset();
-		ProvisionList.Reset();
+		ProvisionList->Reset();
 
 		// Now split up the log into multiple lines
 		TArray<FString> LogLines;
@@ -197,7 +198,12 @@ void FIOSTargetSettingsCustomization::UpdateStatus()
 						Prov->bDistribution = Value.Contains(TEXT("DISTRIBUTION"));
 					}
 				}
-				ProvisionList.Add(Prov);
+
+				// check to see if this the one selected in the ini file
+				FString OutString;
+				MobileProvisionProperty->GetValueAsFormattedString(OutString);
+				Prov->bManuallySelected = (OutString == Prov->FileName);
+				ProvisionList->Add(Prov);
 			}
 			else if (Line.Contains(TEXT("MATCHED-"), ESearchCase::CaseSensitive))
 			{
@@ -244,6 +250,9 @@ void FIOSTargetSettingsCustomization::BuildPListSection(IDetailLayoutBuilder& De
 	IDetailCategoryBuilder& DeviceCategory = DetailLayout.EditCategory(TEXT("Devices"));
 	IDetailCategoryBuilder& BuildCategory = DetailLayout.EditCategory(TEXT("Build"));
 	IDetailCategoryBuilder& ExtraCategory = DetailLayout.EditCategory(TEXT("Extra PList Data"));
+	MobileProvisionProperty = DetailLayout.GetProperty(GET_MEMBER_NAME_CHECKED(UIOSRuntimeSettings, MobileProvision));
+	BuildCategory.AddProperty(MobileProvisionProperty)
+		.Visibility(EVisibility::Hidden);
 
 /*	ProvisionCategory.AddCustomRow(TEXT("Certificate Request"), false)
 		.NameContent()
@@ -341,6 +350,9 @@ void FIOSTargetSettingsCustomization::BuildPListSection(IDetailLayoutBuilder& De
 						.HeaderRow
 						(
 						SNew(SHeaderRow)
+						+ SHeaderRow::Column("Selected")
+						.DefaultLabel(LOCTEXT("ProvisionListNameColumnHeader", ""))
+						.FixedWidth(30.0f)
 						+ SHeaderRow::Column("Name")
 						.DefaultLabel(LOCTEXT("ProvisionListNameColumnHeader", "Provision"))
 						.FillWidth(1.0f)
@@ -1213,7 +1225,19 @@ bool FIOSTargetSettingsCustomization::UpdateStatusDelegate(float DeltaTime)
 TSharedRef<ITableRow> FIOSTargetSettingsCustomization::HandleProvisionListGenerateRow( ProvisionPtr InProvision, const TSharedRef<STableViewBase>& OwnerTable )
 {
 	return SNew(SProvisionListRow, OwnerTable)
-		.Provision(InProvision);
+		.Provision(InProvision)
+		.ProvisionList(ProvisionList)
+		.OnProvisionChanged(this, &FIOSTargetSettingsCustomization::HandleProvisionChanged);
+}
+
+void FIOSTargetSettingsCustomization::HandleProvisionChanged(FString Provision)
+{
+	FText OutText;
+	MobileProvisionProperty->GetValueAsFormattedText(OutText);
+	if (OutText.ToString() != Provision)
+	{
+		MobileProvisionProperty->SetValueFromFormattedString(Provision);
+	}
 }
 
 TSharedRef<ITableRow> FIOSTargetSettingsCustomization::HandleCertificateListGenerateRow( CertificatePtr InCertificate, const TSharedRef<STableViewBase>& OwnerTable )
@@ -1239,23 +1263,23 @@ void FIOSTargetSettingsCustomization::FilterLists()
 	FilteredProvisionList.Reset();
 	FilteredCertificateList.Reset();
 
-	for (int Index = 0; Index < ProvisionList.Num(); ++Index)
+	for (int Index = 0; Index < ProvisionList->Num(); ++Index)
 	{
-		if (SelectedProvision.Contains(ProvisionList[Index]->Name) && SelectedFile.Contains(ProvisionList[Index]->FileName))
+		if (SelectedProvision.Contains((*ProvisionList)[Index]->Name) && SelectedFile.Contains((*ProvisionList)[Index]->FileName))
 		{
-			ProvisionList[Index]->bSelected = true;
+			(*ProvisionList)[Index]->bSelected = true;
 		}
 		else
 		{
-			ProvisionList[Index]->bSelected = false;
+			(*ProvisionList)[Index]->bSelected = false;
 		}
-		if (bShowAllProvisions || ProvisionList[Index]->Status.Contains("VALID"))
+		if (bShowAllProvisions || (*ProvisionList)[Index]->Status.Contains("VALID"))
 		{
-			FilteredProvisionList.Add(ProvisionList[Index]);
+			FilteredProvisionList.Add((*ProvisionList)[Index]);
 		}
 	}
 
-	if (ProvisionList.Num() > 0)
+	if (ProvisionList->Num() > 0)
 	{
 		if (ProvisionInfoSwitcher.IsValid())
 		{
@@ -1263,7 +1287,7 @@ void FIOSTargetSettingsCustomization::FilterLists()
 		}
 		if (FilteredProvisionList.Num() == 0 && !bShowAllProvisions)
 		{
-			FilteredProvisionList.Append(ProvisionList);
+			FilteredProvisionList.Append(*ProvisionList);
 		}
 	}
 	else
diff --git a/Engine/Source/Developer/iOS/IOSPlatformEditor/Private/IOSTargetSettingsCustomization.h b/Engine/Source/Developer/iOS/IOSPlatformEditor/Private/IOSTargetSettingsCustomization.h
--- a/Engine/Source/Developer/iOS/IOSPlatformEditor/Private/IOSTargetSettingsCustomization.h
+++ b/Engine/Source/Developer/iOS/IOSPlatformEditor/Private/IOSTargetSettingsCustomization.h
@@ -16,8 +16,10 @@ public:
 	FString Status;
 	bool bDistribution;
 	bool bSelected;
+	bool bManuallySelected;
 };
 typedef TSharedPtr<class FProvision> ProvisionPtr;
+typedef TSharedPtr<TArray<ProvisionPtr>> ProvisionListPtr;
 
 //////////////////////////////////////////////////////////////////////////
 // FCertificate structure
@@ -68,7 +70,7 @@ private:
 
 	TSharedPtr<FMonitoredProcess> IPPProcess;
 	FDelegateHandle TickerHandle;
-	TArray<ProvisionPtr> ProvisionList;
+	TSharedPtr<TArray<ProvisionPtr>> ProvisionList;
 	TArray<ProvisionPtr> FilteredProvisionList;
 	TSharedPtr<SListView<ProvisionPtr> > ProvisionListView;
 	TSharedPtr<SWidgetSwitcher > ProvisionInfoSwitcher;
@@ -77,6 +79,7 @@ private:
 	TSharedPtr<SListView<CertificatePtr> > CertificateListView;
 	TSharedPtr<SWidgetSwitcher > CertificateInfoSwitcher;
 	TAttribute<bool> RunningIPPProcess;
+	TSharedPtr<IPropertyHandle> MobileProvisionProperty;
 
 	FString SelectedProvision;
 	FString SelectedFile;
@@ -152,6 +155,8 @@ private:
 	// updates the text in the ini file and checks for a valid provision/certificate
 	void OnRemoteServerChanged(const FText& NewText, ETextCommit::Type, TSharedRef<IPropertyHandle> InPropertyHandle);
 
+	void HandleProvisionChanged(FString Provision);
+
 	// 
 	FText GetBundleText(TSharedRef<IPropertyHandle> InPropertyHandle) const;
 };
diff --git a/Engine/Source/Developer/iOS/IOSPlatformEditor/Private/SProvisionListRow.h b/Engine/Source/Developer/iOS/IOSPlatformEditor/Private/SProvisionListRow.h
--- a/Engine/Source/Developer/iOS/IOSPlatformEditor/Private/SProvisionListRow.h
+++ b/Engine/Source/Developer/iOS/IOSPlatformEditor/Private/SProvisionListRow.h
@@ -4,6 +4,8 @@
 
 #define IMAGE_BRUSH( RelativePath, ... ) FSlateImageBrush( FPaths::EngineContentDir() / TEXT("Editor/Slate") / RelativePath + TEXT(".png"), __VA_ARGS__ )
 
+DECLARE_DELEGATE_OneParam(FOnProvisionChanged, FString);
+
 /**
  * Implements a row widget for the provision list view.
  */
@@ -14,7 +16,9 @@ public:
 
 	SLATE_BEGIN_ARGS(SProvisionListRow) { }
 		SLATE_ARGUMENT(ProvisionPtr, Provision)
-	SLATE_END_ARGS()
+		SLATE_ARGUMENT(ProvisionListPtr, ProvisionList)
+		SLATE_EVENT(FOnProvisionChanged, OnProvisionChanged)
+		SLATE_END_ARGS()
 
 public:
 
@@ -26,7 +30,9 @@ public:
 	void Construct( const FArguments& InArgs, const TSharedRef<STableViewBase>& InOwnerTableView )
 	{
 		Provision = InArgs._Provision;
-		
+		ProvisionList = InArgs._ProvisionList;
+		OnProvisionChanged_Handler = InArgs._OnProvisionChanged;
+
 		SMultiColumnTableRow<ProvisionPtr>::Construct(FSuperRowType::FArguments(), InOwnerTableView);
 
 		/* Set images for various SCheckBox states ... */
@@ -59,7 +65,13 @@ public:
 	BEGIN_SLATE_FUNCTION_BUILD_OPTIMIZATION
 	virtual TSharedRef<SWidget> GenerateWidgetForColumn( const FName& ColumnName ) override
 	{
-		if (ColumnName == TEXT("Name"))
+		if (ColumnName == TEXT("Selected"))
+		{
+			return SNew(SCheckBox)
+				.IsChecked(this, &SProvisionListRow::HandleChecked)
+				.OnCheckStateChanged(this, &SProvisionListRow::HandleCheckStateChanged);
+		}
+		else if (ColumnName == TEXT("Name"))
 		{
 			return SNew(SBox)
 				.Padding(FMargin(4.0f, 0.0f))
@@ -109,6 +121,7 @@ public:
 	}
 	END_SLATE_FUNCTION_BUILD_OPTIMIZATION
 
+
 private:
 
 	// Callback for getting the text in the 'Name' column.
@@ -175,10 +188,37 @@ private:
 		return Provision->bDistribution ? ECheckBoxState::Checked : ECheckBoxState::Unchecked;
 	}
 
+	ECheckBoxState HandleChecked() const
+	{
+		return Provision->bManuallySelected ? ECheckBoxState::Checked : ECheckBoxState::Unchecked;
+	}
+
+	void HandleCheckStateChanged(ECheckBoxState InState)
+	{
+		Provision->bManuallySelected = InState == ECheckBoxState::Checked;
+
+		// update the property
+		if (OnProvisionChanged_Handler.IsBound())
+		{
+			OnProvisionChanged_Handler.Execute(Provision->bManuallySelected ? Provision->FileName : "");
+		}
+
+		// disable any other objects
+		for (int32 Idx = 0; Idx < ProvisionList->Num(); ++Idx)
+		{
+			if ((*ProvisionList)[Idx] != Provision && (*ProvisionList)[Idx]->bManuallySelected)
+			{
+				(*ProvisionList)[Idx]->bManuallySelected = false;
+			}
+		}
+	}
+
 private:
 
 	// Holds the target device service used to populate this row.
 	ProvisionPtr Provision;
+	ProvisionListPtr ProvisionList;
+	FOnProvisionChanged OnProvisionChanged_Handler;
 
 	static bool bInitialized;
 	static FCheckBoxStyle ProvisionCheckBoxStyle;
diff --git a/Engine/Source/Programs/AutomationTool/AutomationUtils/ProjectParams.cs b/Engine/Source/Programs/AutomationTool/AutomationUtils/ProjectParams.cs
--- a/Engine/Source/Programs/AutomationTool/AutomationUtils/ProjectParams.cs
+++ b/Engine/Source/Programs/AutomationTool/AutomationUtils/ProjectParams.cs
@@ -405,7 +405,8 @@ namespace AutomationTool
 			bool? FastCook = null,
 			bool? IgnoreCookErrors = null,
 			bool? CodeSign = null,
-			bool? UploadSymbols = null
+			bool? UploadSymbols = null,
+			string Provision = null
 			)
 		{
 			//
@@ -578,6 +579,8 @@ namespace AutomationTool
 				this.DeviceName = this.Device;
 			}
 
+			this.Provision = ParseParamValueIfNotSpecified(Command, Provision, "provision", String.Empty, true);
+
 			this.ServerDevice = ParseParamValueIfNotSpecified(Command, ServerDevice, "serverdevice", this.Device);
 			this.NullRHI = GetParamValueIfNotSpecified(Command, NullRHI, this.NullRHI, "nullrhi");
 			this.FakeClient = GetParamValueIfNotSpecified(Command, FakeClient, this.FakeClient, "fakeclient");
@@ -1284,6 +1287,11 @@ namespace AutomationTool
 		/// </summary>
 		public bool bCodeSign = false;
 
+		/// <summary>
+		/// Provision to use
+		/// </summary>
+		public string Provision = null;
+
 		#endregion
 
 		#region Run
diff --git a/Engine/Source/Programs/AutomationTool/IOS/IOSPlatform.Automation.cs b/Engine/Source/Programs/AutomationTool/IOS/IOSPlatform.Automation.cs
--- a/Engine/Source/Programs/AutomationTool/IOS/IOSPlatform.Automation.cs
+++ b/Engine/Source/Programs/AutomationTool/IOS/IOSPlatform.Automation.cs
@@ -207,6 +207,13 @@ public class IOSPlatform : Platform
 			}
 		}
 
+		if (String.IsNullOrEmpty(Params.Provision))
+		{
+			UnrealBuildTool.IOSPlatform BuildPlat = UEBuildPlatform.GetBuildPlatform(UnrealTargetPlatform.IOS) as UnrealBuildTool.IOSPlatform;
+			BuildPlat.SetUpProjectEnvironment(UnrealTargetPlatform.IOS);
+			Params.Provision = UnrealBuildTool.IOSPlatform.MobileProvision;
+		}
+
 		if (UnrealBuildTool.BuildHostPlatform.Current.Platform != UnrealTargetPlatform.Mac)
 		{
 			var TargetConfiguration = SC.StageTargetConfigurations[0];
@@ -267,6 +274,10 @@ public class IOSPlatform : Platform
 					{
 						IPPArguments += " -iterate";
 					}
+					if (Params.Provision != null)
+					{
+						IPPArguments += " -provision \"" + Params.Provision + "\""; 
+					}
 
 					RunAndLog(CmdEnv, IPPExe, IPPArguments);
 				}
@@ -311,6 +322,11 @@ public class IOSPlatform : Platform
 					{
 						IPPArguments.Add(" -iterate");
 					}
+					if (Params.Provision != null)
+					{
+						IPPArguments.Add(" -provision");
+						IPPArguments.Add(Params.Provision);
+					}
 
 					if (RunIPP(IPPArguments.ToArray()) != 0)
 					{
@@ -350,7 +366,7 @@ public class IOSPlatform : Platform
 				bCreatedIPA = true;
 
 				// code sign the app
-				CodeSign(Path.GetDirectoryName(Params.ProjectGameExeFilename), Params.IsCodeBasedProject ? Params.ShortProjectName : Path.GetFileNameWithoutExtension(Params.ProjectGameExeFilename), Params.RawProjectPath, SC.StageTargetConfigurations[0], SC.LocalRoot, Params.ShortProjectName, Path.GetDirectoryName(Params.RawProjectPath), SC.IsCodeBasedProject, Params.Distribution);
+				CodeSign(Path.GetDirectoryName(Params.ProjectGameExeFilename), Params.IsCodeBasedProject ? Params.ShortProjectName : Path.GetFileNameWithoutExtension(Params.ProjectGameExeFilename), Params.RawProjectPath, SC.StageTargetConfigurations[0], SC.LocalRoot, Params.ShortProjectName, Path.GetDirectoryName(Params.RawProjectPath), SC.IsCodeBasedProject, Params.Distribution, Params.Provision);
 
 				// now generate the ipa
 				PackageIPA(Path.GetDirectoryName(Params.ProjectGameExeFilename), Params.IsCodeBasedProject ? Params.ShortProjectName : Path.GetFileNameWithoutExtension(Params.ProjectGameExeFilename), Params.ShortProjectName, Path.GetDirectoryName(Params.RawProjectPath), SC.StageTargetConfigurations[0], Params.Distribution);
@@ -392,7 +408,7 @@ public class IOSPlatform : Platform
 		return XcodeProj;
 	}
 
-	private void CodeSign(string BaseDirectory, string GameName, string RawProjectPath, UnrealTargetConfiguration TargetConfig, string LocalRoot, string ProjectName, string ProjectDirectory, bool IsCode, bool Distribution = false)
+	private void CodeSign(string BaseDirectory, string GameName, string RawProjectPath, UnrealTargetConfiguration TargetConfig, string LocalRoot, string ProjectName, string ProjectDirectory, bool IsCode, bool Distribution = false, string Provision = null)
 	{
 		// check for the proper xcodeproject
 		bool bWasGenerated = false;
@@ -406,6 +422,26 @@ public class IOSPlatform : Platform
 		Arguments += " -configuration " + TargetConfig.ToString();
 		Arguments += " -sdk iphoneos";
 		Arguments += " CODE_SIGN_IDENTITY=" + (Distribution ? "\"iPhone Distribution\"" : "\"iPhone Developer\"");
+		if (Provision != null)
+		{
+			// read the provision to get the UUID
+			if (File.Exists(Environment.GetEnvironmentVariable("HOME") + "/Library/MobileDevice/Provisioning Profiles/" + Provision))
+			{
+				string UUID = "";
+				string AllText = File.ReadAllText(Environment.GetEnvironmentVariable("HOME") + "/Library/MobileDevice/Provisioning Profiles/" + Provision);
+				int idx = AllText.IndexOf("<key>UUID</key>");
+				if (idx > 0)
+				{
+					idx = AllText.IndexOf("<string>", idx);
+					if (idx > 0)
+					{
+						idx += "<string>".Length;
+						UUID = AllText.Substring(idx, AllText.IndexOf("</string>", idx) - idx);
+						Arguments += " PROVISIONING_PROFILE=" + UUID;
+					}
+				}
+			}
+		}
 		ProcessResult Result = Run ("/usr/bin/env", Arguments, null, ERunOptions.Default);
 		if (bWasGenerated)
 		{
diff --git a/Engine/Source/Programs/IOS/iPhonePackager/MobileProvisionUtilities.cs b/Engine/Source/Programs/IOS/iPhonePackager/MobileProvisionUtilities.cs
--- a/Engine/Source/Programs/IOS/iPhonePackager/MobileProvisionUtilities.cs
+++ b/Engine/Source/Programs/IOS/iPhonePackager/MobileProvisionUtilities.cs
@@ -51,6 +51,14 @@ namespace iPhonePackager
 				Directory.CreateDirectory(Config.ProvisionDirectory);
 			}
 
+			if (Config.bProvision)
+			{
+				if (File.Exists(Config.ProvisionDirectory + "/" + Config.Provision))
+				{
+					return Config.ProvisionDirectory + "/" + Config.Provision;
+				}
+			}
+
 			#region remove after we provide an install mechanism
 			// copy all of the provisions from the game directory to the library
 			if (!String.IsNullOrEmpty(Config.ProjectFile))
diff --git a/Engine/Source/Programs/UnrealBuildTool/Configuration/UEBuildPlatform.cs b/Engine/Source/Programs/UnrealBuildTool/Configuration/UEBuildPlatform.cs
--- a/Engine/Source/Programs/UnrealBuildTool/Configuration/UEBuildPlatform.cs
+++ b/Engine/Source/Programs/UnrealBuildTool/Configuration/UEBuildPlatform.cs
@@ -946,8 +946,6 @@ namespace UnrealBuildTool
 		{
 		}
 
-		static Dictionary<UnrealTargetPlatform, ConfigCacheIni> EngineIniCache = new Dictionary<UnrealTargetPlatform, ConfigCacheIni>();
-
 		protected static bool DoProjectSettingsMatchDefault(UnrealTargetPlatform Platform, DirectoryReference ProjectDirectoryName, string Section, string[] BoolKeys, string[] IntKeys, string[] StringKeys)
 		{
 			ConfigCacheIni ProjIni = ConfigCacheIni.CreateConfigCacheIni(Platform, "Engine", ProjectDirectoryName);
diff --git a/Engine/Source/Programs/UnrealBuildTool/Configuration/UEBuildTarget.cs b/Engine/Source/Programs/UnrealBuildTool/Configuration/UEBuildTarget.cs
--- a/Engine/Source/Programs/UnrealBuildTool/Configuration/UEBuildTarget.cs
+++ b/Engine/Source/Programs/UnrealBuildTool/Configuration/UEBuildTarget.cs
@@ -3354,7 +3354,7 @@ namespace UnrealBuildTool
 			ModuleRules RulesObject = RulesAssembly.CreateModuleRules(ModuleName, TargetInfo, out ModuleFileName);
 
 			// Reads additional dependencies array for project module from project file and fills PrivateDependencyModuleNames. 
-			if(ProjectDescriptor != null)
+			if(ProjectDescriptor != null && ProjectDescriptor.Modules != null)
 			{
 				ModuleDescriptor Module = ProjectDescriptor.Modules.FirstOrDefault(x => x.Name.Equals(ModuleName, StringComparison.InvariantCultureIgnoreCase));
 				if(Module != null && Module.AdditionalDependencies != null)
diff --git a/Engine/Source/Programs/UnrealBuildTool/IOS/UEBuildIOS.cs b/Engine/Source/Programs/UnrealBuildTool/IOS/UEBuildIOS.cs
--- a/Engine/Source/Programs/UnrealBuildTool/IOS/UEBuildIOS.cs
+++ b/Engine/Source/Programs/UnrealBuildTool/IOS/UEBuildIOS.cs
@@ -9,7 +9,7 @@ using System.Xml;
 
 namespace UnrealBuildTool
 {
-    class IOSPlatform : UEBuildPlatform
+    public class IOSPlatform : UEBuildPlatform
     {
         // by default, use an empty architecture (which is really just a modifer to the platform for some paths/names)
         [XmlConfig]
@@ -33,6 +33,9 @@ namespace UnrealBuildTool
 		/** additional linker flags for non-shipping */
 		public static string AdditionalLinkerFlags = "";
 
+		/** mobile provision to use for code signing */
+		public static string MobileProvision = "";
+
 		private bool bInitializedProject = false;
 
 		public string GetRunTimeVersion()
@@ -260,6 +263,8 @@ namespace UnrealBuildTool
 				Ini.GetString("/Script/IOSRuntimeSettings.IOSRuntimeSettings", "AdditionalLinkerFlags", out AdditionalLinkerFlags);
 				Ini.GetString("/Script/IOSRuntimeSettings.IOSRuntimeSettings", "AdditionalShippingLinkerFlags", out AdditionalShippingLinkerFlags);
 
+				Ini.GetString("/Script/IOSRuntimeSettings.IOSRuntimeSettings", "MobileProvision", out MobileProvision);
+
 				bInitializedProject = true;
 			}
 		}
diff --git a/Engine/Source/Runtime/IOS/IOSRuntimeSettings/Classes/IOSRuntimeSettings.h b/Engine/Source/Runtime/IOS/IOSRuntimeSettings/Classes/IOSRuntimeSettings.h
--- a/Engine/Source/Runtime/IOS/IOSRuntimeSettings/Classes/IOSRuntimeSettings.h
+++ b/Engine/Source/Runtime/IOS/IOSRuntimeSettings/Classes/IOSRuntimeSettings.h
@@ -286,6 +286,11 @@ public:
 	UPROPERTY(GlobalConfig, EditAnywhere, Category = Online, meta = (EditCondition = "bEnableFacebookSupport"))
 	FString FacebookAppID;
 
+	// Mobile provision to utilize when signing
+	UPROPERTY(GlobalConfig, EditAnywhere, Category = Build)
+	FString MobileProvision;
+
+
 #if WITH_EDITOR
 	// UObject interface
 	virtual void PostEditChangeProperty(struct FPropertyChangedEvent& PropertyChangedEvent) override;
