Merging mobile provision selection from 4.10

diff --git a/Engine/Source/Developer/iOS/IOSPlatformEditor/Private/IOSTargetSettingsCustomization.cpp b/Engine/Source/Developer/iOS/IOSPlatformEditor/Private/IOSTargetSettingsCustomization.cpp
--- a/Engine/Source/Developer/iOS/IOSPlatformEditor/Private/IOSTargetSettingsCustomization.cpp
+++ b/Engine/Source/Developer/iOS/IOSPlatformEditor/Private/IOSTargetSettingsCustomization.cpp
@@ -77,6 +77,7 @@ FIOSTargetSettingsCustomization::FIOSTargetSettingsCustomization()
 	bShowAllProvisions = false;
 	bShowAllCertificates = false;
 	ProvisionList = MakeShareable(new TArray<ProvisionPtr>());
+	CertificateList = MakeShareable(new TArray<CertificatePtr>());
 }
 
 FIOSTargetSettingsCustomization::~FIOSTargetSettingsCustomization()
@@ -150,17 +151,22 @@ void FIOSTargetSettingsCustomization::UpdateStatus()
 					}
 				}
 				CertificatePtr PrevCert = NULL;
-				for (int CIndex = 0; CIndex < CertificateList.Num() && !PrevCert.IsValid(); ++CIndex)
+				for (int CIndex = 0; CIndex < CertificateList->Num() && !PrevCert.IsValid(); ++CIndex)
 				{
-					if (CertificateList[CIndex]->Name == Cert->Name)
+					if ((*CertificateList)[CIndex]->Name == Cert->Name)
 					{
-						PrevCert = CertificateList[CIndex];
+						PrevCert = (*CertificateList)[CIndex];
 						break;
 					}
 				}
+
+				// check to see if this the one selected in the ini file
+				FString OutString;
+				SignCertificateProperty->GetValueAsFormattedString(OutString);
+				Cert->bManuallySelected = (OutString == Cert->Name);
 				if (!PrevCert.IsValid())
 				{
-					CertificateList.Add(Cert);
+					CertificateList->Add(Cert);
 				}
 				else
 				{
@@ -253,6 +259,9 @@ void FIOSTargetSettingsCustomization::BuildPListSection(IDetailLayoutBuilder& De
 	MobileProvisionProperty = DetailLayout.GetProperty(GET_MEMBER_NAME_CHECKED(UIOSRuntimeSettings, MobileProvision));
 	BuildCategory.AddProperty(MobileProvisionProperty)
 		.Visibility(EVisibility::Hidden);
+	SignCertificateProperty = DetailLayout.GetProperty(GET_MEMBER_NAME_CHECKED(UIOSRuntimeSettings, SigningCertificate));
+	BuildCategory.AddProperty(SignCertificateProperty)
+		.Visibility(EVisibility::Hidden);
 
 /*	ProvisionCategory.AddCustomRow(TEXT("Certificate Request"), false)
 		.NameContent()
@@ -351,7 +360,7 @@ void FIOSTargetSettingsCustomization::BuildPListSection(IDetailLayoutBuilder& De
 						(
 						SNew(SHeaderRow)
 						+ SHeaderRow::Column("Selected")
-						.DefaultLabel(LOCTEXT("ProvisionListNameColumnHeader", ""))
+						.DefaultLabel(LOCTEXT("ProvisionListSelectColumnHeader", ""))
 						.FixedWidth(30.0f)
 						+ SHeaderRow::Column("Name")
 						.DefaultLabel(LOCTEXT("ProvisionListNameColumnHeader", "Provision"))
@@ -511,6 +520,9 @@ void FIOSTargetSettingsCustomization::BuildPListSection(IDetailLayoutBuilder& De
 								.HeaderRow
 								(
 									SNew(SHeaderRow)
+									+ SHeaderRow::Column("Selected")
+									.DefaultLabel(LOCTEXT("CertificateListSelectColumnHeader", ""))
+									.FixedWidth(30.0f)
 									+ SHeaderRow::Column("Name")
 									.DefaultLabel(LOCTEXT("CertificateListNameColumnHeader", "Certificate"))
 									+ SHeaderRow::Column("Status")
@@ -1240,10 +1252,22 @@ void FIOSTargetSettingsCustomization::HandleProvisionChanged(FString Provision)
 	}
 }
 
+void FIOSTargetSettingsCustomization::HandleCertificateChanged(FString Certificate)
+{
+	FText OutText;
+	SignCertificateProperty->GetValueAsFormattedText(OutText);
+	if (OutText.ToString() != Certificate)
+	{
+		SignCertificateProperty->SetValueFromFormattedString(Certificate);
+	}
+}
+
 TSharedRef<ITableRow> FIOSTargetSettingsCustomization::HandleCertificateListGenerateRow( CertificatePtr InCertificate, const TSharedRef<STableViewBase>& OwnerTable )
 {
 	return SNew(SCertificateListRow, OwnerTable)
-		.Certificate(InCertificate);
+		.Certificate(InCertificate)
+		.CertificateList(CertificateList)
+		.OnCertificateChanged(this, &FIOSTargetSettingsCustomization::HandleCertificateChanged);
 }
 
 void FIOSTargetSettingsCustomization::HandleAllProvisionsHyperlinkNavigate( bool AllProvisions )
@@ -1298,23 +1322,23 @@ void FIOSTargetSettingsCustomization::FilterLists()
 		}
 	}
 
-	for (int Index = 0; Index < CertificateList.Num(); ++Index)
+	for (int Index = 0; Index < CertificateList->Num(); ++Index)
 	{
-		if (SelectedCert.Contains(CertificateList[Index]->Name))
+		if (SelectedCert.Contains((*CertificateList)[Index]->Name))
 		{
-			CertificateList[Index]->bSelected = true;
+			(*CertificateList)[Index]->bSelected = true;
 		}
 		else
 		{
-			CertificateList[Index]->bSelected = false;
+			(*CertificateList)[Index]->bSelected = false;
 		}
-		if (bShowAllCertificates || CertificateList[Index]->Status.Contains("VALID"))
+		if (bShowAllCertificates || (*CertificateList)[Index]->Status.Contains("VALID"))
 		{
-			FilteredCertificateList.Add(CertificateList[Index]);
+			FilteredCertificateList.Add((*CertificateList)[Index]);
 		}
 	}
 
-	if (CertificateList.Num() > 0)
+	if (CertificateList->Num() > 0)
 	{
 		if (CertificateInfoSwitcher.IsValid())
 		{
@@ -1322,7 +1346,7 @@ void FIOSTargetSettingsCustomization::FilterLists()
 		}
 		if (FilteredCertificateList.Num() == 0 && !bShowAllCertificates)
 		{
-			FilteredCertificateList.Append(CertificateList);
+			FilteredCertificateList.Append(*CertificateList);
 		}
 	}
 	else
diff --git a/Engine/Source/Developer/iOS/IOSPlatformEditor/Private/IOSTargetSettingsCustomization.h b/Engine/Source/Developer/iOS/IOSPlatformEditor/Private/IOSTargetSettingsCustomization.h
--- a/Engine/Source/Developer/iOS/IOSPlatformEditor/Private/IOSTargetSettingsCustomization.h
+++ b/Engine/Source/Developer/iOS/IOSPlatformEditor/Private/IOSTargetSettingsCustomization.h
@@ -30,8 +30,10 @@ public:
 	FString Status;
 	FString Expires;
 	bool bSelected;
+	bool bManuallySelected;
 };
 typedef TSharedPtr<class FCertificate> CertificatePtr;
+typedef TSharedPtr<TArray<CertificatePtr>> CertificateListPtr;
 
 //////////////////////////////////////////////////////////////////////////
 // FIOSTargetSettingsCustomization
@@ -74,12 +76,13 @@ private:
 	TArray<ProvisionPtr> FilteredProvisionList;
 	TSharedPtr<SListView<ProvisionPtr> > ProvisionListView;
 	TSharedPtr<SWidgetSwitcher > ProvisionInfoSwitcher;
-	TArray<CertificatePtr> CertificateList;
+	TSharedPtr<TArray<CertificatePtr>> CertificateList;
 	TArray<CertificatePtr> FilteredCertificateList;
 	TSharedPtr<SListView<CertificatePtr> > CertificateListView;
 	TSharedPtr<SWidgetSwitcher > CertificateInfoSwitcher;
 	TAttribute<bool> RunningIPPProcess;
 	TSharedPtr<IPropertyHandle> MobileProvisionProperty;
+	TSharedPtr<IPropertyHandle> SignCertificateProperty;
 
 	FString SelectedProvision;
 	FString SelectedFile;
@@ -157,6 +160,8 @@ private:
 
 	void HandleProvisionChanged(FString Provision);
 
+	void HandleCertificateChanged(FString Certificate);
+
 	// 
 	FText GetBundleText(TSharedRef<IPropertyHandle> InPropertyHandle) const;
 };
diff --git a/Engine/Source/Programs/AutomationTool/AutomationUtils/ProjectParams.cs b/Engine/Source/Programs/AutomationTool/AutomationUtils/ProjectParams.cs
--- a/Engine/Source/Programs/AutomationTool/AutomationUtils/ProjectParams.cs
+++ b/Engine/Source/Programs/AutomationTool/AutomationUtils/ProjectParams.cs
@@ -405,6 +405,7 @@ namespace AutomationTool
 			bool? CodeSign = null,
 			bool? UploadSymbols = null,
 			string Provision = null,
+			string Certificate = null,
 			ParamList<string> InMapsToRebuildLightMaps = null
 			)
 		{
@@ -578,6 +579,7 @@ namespace AutomationTool
 			}
 
 			this.Provision = ParseParamValueIfNotSpecified(Command, Provision, "provision", String.Empty, true);
+			this.Certificate = ParseParamValueIfNotSpecified(Command, Certificate, "certificate", String.Empty, true);
 
 			this.ServerDevice = ParseParamValueIfNotSpecified(Command, ServerDevice, "serverdevice", this.Device);
 			this.NullRHI = GetParamValueIfNotSpecified(Command, NullRHI, this.NullRHI, "nullrhi");
@@ -993,15 +995,12 @@ namespace AutomationTool
         /// <summary>
         /// Shared: Ref to an auto-generated plugin file that should be incorporated into the project's build
         /// </summary>
-        public FileReference NativizedScriptPlugin;
+        public FileReference GeneratedPluginDescFile;
 
-        /// <summary>
-        /// Shared: Used to guard against fatal use of the GeneratedScriptPlugin param
-        /// </summary>
-        public bool UseNativizedScriptPlugin()
-        {
-            return (NativizedScriptPlugin != null) && NativizedScriptPlugin.Exists();
-        }
+		/// <summary>
+		/// Shared: Ref to a manifest file, detailing the asset files that were converted into native source
+		/// </summary>
+		public FileReference GeneratedPluginManifest;
 
 		#endregion
 
@@ -1320,6 +1319,11 @@ namespace AutomationTool
 		/// </summary>
 		public string Provision = null;
 
+		/// <summary>
+		/// Certificate to use
+		/// </summary>
+		public string Certificate = null;
+
 		#endregion
 
 		#region Run
diff --git a/Engine/Source/Programs/AutomationTool/IOS/IOSPlatform.Automation.cs b/Engine/Source/Programs/AutomationTool/IOS/IOSPlatform.Automation.cs
--- a/Engine/Source/Programs/AutomationTool/IOS/IOSPlatform.Automation.cs
+++ b/Engine/Source/Programs/AutomationTool/IOS/IOSPlatform.Automation.cs
@@ -215,6 +215,13 @@ public class IOSPlatform : Platform
 			BuildPlatContext.SetUpProjectEnvironment();
 			Params.Provision = BuildPlatContext.MobileProvision;
 		}
+		if (String.IsNullOrEmpty(Params.Certificate))
+		{
+			UnrealBuildTool.IOSPlatform BuildPlat = UEBuildPlatform.GetBuildPlatform(UnrealTargetPlatform.IOS) as UnrealBuildTool.IOSPlatform;
+			UnrealBuildTool.IOSPlatformContext BuildPlatContext = (IOSPlatformContext)BuildPlat.CreateContext(Params.RawProjectPath);
+			BuildPlatContext.SetUpProjectEnvironment();
+			Params.Certificate = BuildPlatContext.SigningCertificate;
+		}
 
 		if (UnrealBuildTool.BuildHostPlatform.Current.Platform != UnrealTargetPlatform.Mac)
 		{
@@ -280,6 +287,10 @@ public class IOSPlatform : Platform
 					{
 						IPPArguments += " -provision \"" + Params.Provision + "\""; 
 					}
+					if (Params.Certificate != null)
+					{
+						IPPArguments += " -certificate \"" + Params.Certificate + "\"";
+					}
 
 					RunAndLog(CmdEnv, IPPExe, IPPArguments);
 				}
@@ -329,6 +340,11 @@ public class IOSPlatform : Platform
 						IPPArguments.Add(" -provision");
 						IPPArguments.Add(Params.Provision);
 					}
+					if (Params.Certificate != null)
+					{
+						IPPArguments.Add(" -certificate");
+						IPPArguments.Add(Params.Certificate);
+					}
 
 					if (RunIPP(IPPArguments.ToArray()) != 0)
 					{
@@ -368,7 +384,7 @@ public class IOSPlatform : Platform
 				bCreatedIPA = true;
 
 				// code sign the app
-				CodeSign(Path.GetDirectoryName(Params.ProjectGameExeFilename), Params.IsCodeBasedProject ? Params.ShortProjectName : Path.GetFileNameWithoutExtension(Params.ProjectGameExeFilename), Params.RawProjectPath, SC.StageTargetConfigurations[0], SC.LocalRoot, Params.ShortProjectName, Path.GetDirectoryName(Params.RawProjectPath.FullName), SC.IsCodeBasedProject, Params.Distribution, Params.Provision);
+				CodeSign(Path.GetDirectoryName(Params.ProjectGameExeFilename), Params.IsCodeBasedProject ? Params.ShortProjectName : Path.GetFileNameWithoutExtension(Params.ProjectGameExeFilename), Params.RawProjectPath, SC.StageTargetConfigurations[0], SC.LocalRoot, Params.ShortProjectName, Path.GetDirectoryName(Params.RawProjectPath.FullName), SC.IsCodeBasedProject, Params.Distribution, Params.Provision, Params.Certificate);
 
 				// now generate the ipa
 				PackageIPA(Path.GetDirectoryName(Params.ProjectGameExeFilename), Params.IsCodeBasedProject ? Params.ShortProjectName : Path.GetFileNameWithoutExtension(Params.ProjectGameExeFilename), Params.ShortProjectName, Path.GetDirectoryName(Params.RawProjectPath.FullName), SC.StageTargetConfigurations[0], Params.Distribution);
@@ -410,7 +426,7 @@ public class IOSPlatform : Platform
 		return XcodeProj;
 	}
 
-	private void CodeSign(string BaseDirectory, string GameName, FileReference RawProjectPath, UnrealTargetConfiguration TargetConfig, string LocalRoot, string ProjectName, string ProjectDirectory, bool IsCode, bool Distribution = false, string Provision = null)
+	private void CodeSign(string BaseDirectory, string GameName, FileReference RawProjectPath, UnrealTargetConfiguration TargetConfig, string LocalRoot, string ProjectName, string ProjectDirectory, bool IsCode, bool Distribution = false, string Provision = null, string Certificate = null)
 	{
 		// check for the proper xcodeproject
 		bool bWasGenerated = false;
@@ -423,7 +439,14 @@ public class IOSPlatform : Platform
 		Arguments += " - iOS'";
 		Arguments += " -configuration " + TargetConfig.ToString();
 		Arguments += " -sdk iphoneos";
-		Arguments += " CODE_SIGN_IDENTITY=" + (Distribution ? "\"iPhone Distribution\"" : "\"iPhone Developer\"");
+		if (Certificate != null)
+		{
+			Arguments += " CODE_SIGN_IDENTITY=" + Certificate;
+		}
+		else
+		{
+			Arguments += " CODE_SIGN_IDENTITY=" + (Distribution ? "\"iPhone Distribution\"" : "\"iPhone Developer\"");
+		}
 		if (Provision != null)
 		{
 			// read the provision to get the UUID
diff --git a/Engine/Source/Programs/IOS/iPhonePackager/MobileProvisionUtilities.cs b/Engine/Source/Programs/IOS/iPhonePackager/MobileProvisionUtilities.cs
--- a/Engine/Source/Programs/IOS/iPhonePackager/MobileProvisionUtilities.cs
+++ b/Engine/Source/Programs/IOS/iPhonePackager/MobileProvisionUtilities.cs
@@ -149,7 +149,11 @@ namespace iPhonePackager
 					}
 					else
 					{
-						bPassesNameCheck = TestProvision.ApplicationIdentifier.Contains("*");
+						if (TestProvision.ApplicationIdentifier.Contains("*"))
+						{
+							string CompanyName = TestProvision.ApplicationIdentifier.Substring(TestProvision.ApplicationIdentifierPrefix.Length + 1);
+							bPassesNameCheck = CompanyName == "*";
+						}
 					}
 					if (!bPassesNameCheck && bCheckIdentifier)
 					{
@@ -177,7 +181,12 @@ namespace iPhonePackager
 					bool bPassesHasMatchingCertCheck = false;
 					if (bCheckCert)
 					{
-						bPassesHasMatchingCertCheck = (CodeSignatureBuilder.FindCertificate(TestProvision) != null);
+						X509Certificate2 Cert = CodeSignatureBuilder.FindCertificate(TestProvision);
+						bPassesHasMatchingCertCheck = (Cert != null);
+						if (bPassesHasMatchingCertCheck && Config.bCert)
+						{
+							bPassesHasMatchingCertCheck &= (Cert.FriendlyName == Config.Certificate);
+						}
 					}
 					else
 					{
diff --git a/Engine/Source/Programs/UnrealBuildTool/IOS/UEBuildIOS.cs b/Engine/Source/Programs/UnrealBuildTool/IOS/UEBuildIOS.cs
--- a/Engine/Source/Programs/UnrealBuildTool/IOS/UEBuildIOS.cs
+++ b/Engine/Source/Programs/UnrealBuildTool/IOS/UEBuildIOS.cs
@@ -44,6 +44,11 @@ namespace UnrealBuildTool
 		/// </summary>
 		public string MobileProvision = "";
 
+		/// <summary>
+		/// signing certificate to use for code signing
+		/// </summary>
+		public string SigningCertificate = "";
+
 		public IOSPlatformContext(FileReference InProjectFile) : base(UnrealTargetPlatform.IOS, InProjectFile)
 		{
 		}
@@ -179,6 +184,7 @@ namespace UnrealBuildTool
 				Ini.GetString("/Script/IOSRuntimeSettings.IOSRuntimeSettings", "AdditionalShippingLinkerFlags", out AdditionalShippingLinkerFlags);
 
 				Ini.GetString("/Script/IOSRuntimeSettings.IOSRuntimeSettings", "MobileProvision", out MobileProvision);
+				Ini.GetString("/Script/IOSRuntimeSettings.IOSRuntimeSettings", "SigningCertificate", out SigningCertificate);
 
 				bInitializedProject = true;
 			}
diff --git a/Engine/Source/Runtime/IOS/IOSRuntimeSettings/Classes/IOSRuntimeSettings.h b/Engine/Source/Runtime/IOS/IOSRuntimeSettings/Classes/IOSRuntimeSettings.h
--- a/Engine/Source/Runtime/IOS/IOSRuntimeSettings/Classes/IOSRuntimeSettings.h
+++ b/Engine/Source/Runtime/IOS/IOSRuntimeSettings/Classes/IOSRuntimeSettings.h
@@ -290,6 +290,10 @@ public:
 	UPROPERTY(GlobalConfig, EditAnywhere, Category = Build)
 	FString MobileProvision;
 
+	// Signing certificate to utilize when signing
+	UPROPERTY(GlobalConfig, EditAnywhere, Category = Build)
+	FString SigningCertificate;
+
 
 #if WITH_EDITOR
 	// UObject interface
