Merging missed file from 4.10

diff --git a/Engine/Source/Developer/iOS/IOSPlatformEditor/Private/SCertificateListRow.h b/Engine/Source/Developer/iOS/IOSPlatformEditor/Private/SCertificateListRow.h
--- a/Engine/Source/Developer/iOS/IOSPlatformEditor/Private/SCertificateListRow.h
+++ b/Engine/Source/Developer/iOS/IOSPlatformEditor/Private/SCertificateListRow.h
@@ -2,6 +2,8 @@
 
 #pragma once
 
+DECLARE_DELEGATE_OneParam(FOnCertificateChanged, FString);
+
 /**
  * Implements a row widget for the certificate list view.
  */
@@ -12,6 +14,8 @@ public:
 
 	SLATE_BEGIN_ARGS(SCertificateListRow) { }
 		SLATE_ARGUMENT(CertificatePtr, Certificate)
+		SLATE_ARGUMENT(CertificateListPtr, CertificateList)
+		SLATE_EVENT(FOnCertificateChanged, OnCertificateChanged)
 	SLATE_END_ARGS()
 
 public:
@@ -24,6 +28,8 @@ public:
 	void Construct( const FArguments& InArgs, const TSharedRef<STableViewBase>& InOwnerTableView )
 	{
 		Certificate = InArgs._Certificate;
+		CertificateList = InArgs._CertificateList;
+		OnCertificateChanged_Handler = InArgs._OnCertificateChanged;
 		
 		SMultiColumnTableRow<CertificatePtr>::Construct(FSuperRowType::FArguments(), InOwnerTableView);
 	}
@@ -39,7 +45,13 @@ public:
 	BEGIN_SLATE_FUNCTION_BUILD_OPTIMIZATION
 	virtual TSharedRef<SWidget> GenerateWidgetForColumn( const FName& ColumnName ) override
 	{
-		if (ColumnName == TEXT("Name"))
+		if (ColumnName == TEXT("Selected"))
+		{
+			return SNew(SCheckBox)
+				.IsChecked(this, &SCertificateListRow::HandleChecked)
+				.OnCheckStateChanged(this, &SCertificateListRow::HandleCheckStateChanged);
+		}
+		else if (ColumnName == TEXT("Name"))
 		{
 			return SNew(SBox)
 				.Padding(FMargin(4.0f, 0.0f))
@@ -126,8 +138,35 @@ private:
 		return FText::FromString(Certificate->Expires);
 	}
 
+	ECheckBoxState HandleChecked() const
+	{
+		return Certificate->bManuallySelected ? ECheckBoxState::Checked : ECheckBoxState::Unchecked;
+	}
+
+	void HandleCheckStateChanged(ECheckBoxState InState)
+	{
+		Certificate->bManuallySelected = InState == ECheckBoxState::Checked;
+
+		// update the property
+		if (OnCertificateChanged_Handler.IsBound())
+		{
+			OnCertificateChanged_Handler.Execute(Certificate->bManuallySelected ? Certificate->Name : "");
+		}
+
+		// disable any other objects
+		for (int32 Idx = 0; Idx < CertificateList->Num(); ++Idx)
+		{
+			if ((*CertificateList)[Idx] != Certificate && (*CertificateList)[Idx]->bManuallySelected)
+			{
+				(*CertificateList)[Idx]->bManuallySelected = false;
+			}
+		}
+	}
+
 private:
 
 	// Holds the target device service used to populate this row.
 	CertificatePtr Certificate;
+	CertificateListPtr CertificateList;
+	FOnCertificateChanged OnCertificateChanged_Handler;
 };
