fix for assert when bringing up the iOS project settings
UE-22000

diff --git a/Engine/Source/Developer/iOS/IOSPlatformEditor/Private/IOSTargetSettingsCustomization.cpp b/Engine/Source/Developer/iOS/IOSPlatformEditor/Private/IOSTargetSettingsCustomization.cpp
--- a/Engine/Source/Developer/iOS/IOSPlatformEditor/Private/IOSTargetSettingsCustomization.cpp
+++ b/Engine/Source/Developer/iOS/IOSPlatformEditor/Private/IOSTargetSettingsCustomization.cpp
@@ -114,7 +114,7 @@ void FIOSTargetSettingsCustomization::UpdateStatus()
 {
 	if (OutputMessage.Len() > 0)
 	{
-		CertificateList.Reset();
+		CertificateList->Reset();
 		ProvisionList->Reset();
 
 		// Now split up the log into multiple lines
