Fixing case of files in "iOS" directory (pt 2)

diff --git a/Engine/Source/Developer/IOS/IOSPlatformEditor/IOSPlatformEditor.Build.cs b/Engine/Source/Developer/IOS/IOSPlatformEditor/IOSPlatformEditor.Build.cs
new file mode 100644
--- /dev/null
+++ b/Engine/Source/Developer/IOS/IOSPlatformEditor/IOSPlatformEditor.Build.cs
@@ -0,0 +1,48 @@
+// Copyright 1998-2017 Epic Games, Inc. All Rights Reserved.
+
+using UnrealBuildTool;
+
+public class IOSPlatformEditor : ModuleRules
+{
+	public IOSPlatformEditor(ReadOnlyTargetRules Target) : base(Target)
+	{
+		BinariesSubFolder = "IOS";
+
+		PrivateDependencyModuleNames.AddRange(
+			new string[] {
+				"Core",
+				"CoreUObject",
+				"InputCore",
+				"DesktopPlatform",
+				"Engine",
+				"MainFrame",
+				"Slate",
+				"SlateCore",
+				"EditorStyle",
+				"PropertyEditor",
+				"SharedSettingsWidgets",
+				"SourceControl",
+				"IOSRuntimeSettings",
+				"TargetPlatform",
+				"MaterialShaderQualitySettings",
+				"RenderCore",
+			}
+		);
+
+		PrivateIncludePathModuleNames.AddRange(
+			new string[] {
+				"GameProjectGeneration",
+				"Settings",
+			}
+		);
+
+		DynamicallyLoadedModuleNames.AddRange(
+			new string[] {
+				"GameProjectGeneration",
+				}
+		);
+
+        // this is listed above, so it isn't really dynamically loaded, this just marks it as being platform specific.
+		PlatformSpecificDynamicallyLoadedModuleNames.Add("IOSRuntimeSettings");
+	}
+}
diff --git a/Engine/Source/Developer/IOS/IOSPlatformEditor/Private/IOSPlatformEditorModule.cpp b/Engine/Source/Developer/IOS/IOSPlatformEditor/Private/IOSPlatformEditorModule.cpp
new file mode 100644
--- /dev/null
+++ b/Engine/Source/Developer/IOS/IOSPlatformEditor/Private/IOSPlatformEditorModule.cpp
@@ -0,0 +1,101 @@
+// Copyright 1998-2017 Epic Games, Inc. All Rights Reserved.
+
+#include "CoreMinimal.h"
+#include "Modules/ModuleInterface.h"
+#include "Modules/ModuleManager.h"
+#include "Materials/Material.h"
+#include "Materials/MaterialInstance.h"
+#include "PropertyEditorModule.h"
+#include "IOSRuntimeSettings.h"
+#include "IOSTargetSettingsCustomization.h"
+#include "ISettingsModule.h"
+#include "MaterialShaderQualitySettings.h"
+#include "MaterialShaderQualitySettingsCustomization.h"
+#include "ComponentRecreateRenderStateContext.h"
+#include "ShaderPlatformQualitySettings.h"
+
+#define LOCTEXT_NAMESPACE "FIOSPlatformEditorModule"
+
+
+/**
+ * Module for iOS as a target platform
+ */
+class FIOSPlatformEditorModule
+	: public IModuleInterface
+{
+	// IModuleInterface interface
+
+	virtual void StartupModule() override
+	{
+		// register settings detail panel customization
+		FPropertyEditorModule& PropertyModule = FModuleManager::LoadModuleChecked<FPropertyEditorModule>("PropertyEditor");
+		PropertyModule.RegisterCustomClassLayout(
+			"IOSRuntimeSettings",
+			FOnGetDetailCustomizationInstance::CreateStatic(&FIOSTargetSettingsCustomization::MakeInstance)
+		);
+
+		FOnUpdateMaterialShaderQuality UpdateMaterials = FOnUpdateMaterialShaderQuality::CreateLambda([]()
+		{
+			FGlobalComponentRecreateRenderStateContext Recreate;
+			FlushRenderingCommands();
+			UMaterial::AllMaterialsCacheResourceShadersForRendering();
+			UMaterialInstance::AllMaterialsCacheResourceShadersForRendering();
+		});
+
+		PropertyModule.RegisterCustomClassLayout(
+			UShaderPlatformQualitySettings::StaticClass()->GetFName(),
+			FOnGetDetailCustomizationInstance::CreateStatic(&FMaterialShaderQualitySettingsCustomization::MakeInstance, UpdateMaterials)
+			);
+
+		PropertyModule.NotifyCustomizationModuleChanged();
+
+		// register settings
+		ISettingsModule* SettingsModule = FModuleManager::GetModulePtr<ISettingsModule>("Settings");
+
+		if (SettingsModule != nullptr)
+		{
+			SettingsModule->RegisterSettings("Project", "Platforms", "iOS",
+				LOCTEXT("RuntimeSettingsName", "iOS"),
+				LOCTEXT("RuntimeSettingsDescription", "Settings and resources for the iOS platform"),
+				GetMutableDefault<UIOSRuntimeSettings>()
+			);
+
+			{
+				static FName NAME_OPENGL_ES2_IOS(TEXT("GLSL_ES2_IOS"));
+				const UShaderPlatformQualitySettings* IOSMaterialQualitySettings = UMaterialShaderQualitySettings::Get()->GetShaderPlatformQualitySettings(NAME_OPENGL_ES2_IOS);
+				SettingsModule->RegisterSettings("Project", "Platforms", "iOSES2Quality",
+					LOCTEXT("IOSES2QualitySettingsName", "iOS Material Quality - ES2"),
+					LOCTEXT("IOSES2QualitySettingsDescription", "Settings for iOS ES2 material quality"),
+					IOSMaterialQualitySettings
+					);
+			}
+
+			{
+				static FName NAME_SF_METAL(TEXT("SF_METAL"));
+				const UShaderPlatformQualitySettings* IOSMaterialQualitySettings = UMaterialShaderQualitySettings::Get()->GetShaderPlatformQualitySettings(NAME_SF_METAL);
+				SettingsModule->RegisterSettings("Project", "Platforms", "iOSMetalQuality",
+					LOCTEXT("IOSMetalQualitySettingsName", "iOS Material Quality - Metal"),
+					LOCTEXT("IOSMetalQualitySettingsDescription", "Settings for iOS Metal material quality"),
+					IOSMaterialQualitySettings
+				);
+			}
+		}
+	}
+
+	virtual void ShutdownModule() override
+	{
+		ISettingsModule* SettingsModule = FModuleManager::GetModulePtr<ISettingsModule>("Settings");
+
+		if (SettingsModule != nullptr)
+		{
+			SettingsModule->UnregisterSettings("Project", "Platforms", "iOS");
+			SettingsModule->UnregisterSettings("Project", "Platforms", "iOSES2Quality");
+			SettingsModule->UnregisterSettings("Project", "Platforms", "iOSMetalQuality");
+		}
+	}
+};
+
+
+IMPLEMENT_MODULE(FIOSPlatformEditorModule, IOSPlatformEditor);
+
+#undef LOCTEXT_NAMESPACE
diff --git a/Engine/Source/Developer/IOS/IOSPlatformEditor/Private/IOSTargetSettingsCustomization.cpp b/Engine/Source/Developer/IOS/IOSPlatformEditor/Private/IOSTargetSettingsCustomization.cpp
new file mode 100644
--- /dev/null
+++ b/Engine/Source/Developer/IOS/IOSPlatformEditor/Private/IOSTargetSettingsCustomization.cpp
@@ -0,0 +1,1748 @@
+// Copyright 1998-2017 Epic Games, Inc. All Rights Reserved.
+
+#include "IOSTargetSettingsCustomization.h"
+#include "Containers/Ticker.h"
+#include "Misc/App.h"
+#include "Misc/MonitoredProcess.h"
+#include "Widgets/Layout/SSeparator.h"
+#include "Widgets/Text/STextBlock.h"
+#include "Widgets/Text/SRichTextBlock.h"
+#include "Widgets/Input/SEditableTextBox.h"
+#include "Widgets/Input/SButton.h"
+#include "Widgets/Layout/SWidgetSwitcher.h"
+#include "EditorStyleSet.h"
+#include "IOSRuntimeSettings.h"
+#include "PropertyHandle.h"
+#include "DetailLayoutBuilder.h"
+#include "DetailWidgetRow.h"
+#include "IDetailPropertyRow.h"
+#include "DetailCategoryBuilder.h"
+#include "DesktopPlatformModule.h"
+#include "Framework/Application/SlateApplication.h"
+#include "IDetailGroup.h"
+#include "SExternalImageReference.h"
+#include "PlatformIconInfo.h"
+#include "SourceControlHelpers.h"
+#include "Framework/Notifications/NotificationManager.h"
+#include "Widgets/Notifications/SNotificationList.h"
+#include "Widgets/Input/SHyperlink.h"
+#include "Widgets/Views/SListView.h"
+#include "SProvisionListRow.h"
+#include "SCertificateListRow.h"
+#include "Misc/EngineBuildSettings.h"
+#include "SNumericDropDown.h"
+#include "Misc/MessageDialog.h"
+#include "Widgets/Notifications/SErrorText.h"
+#include "Interfaces/ITargetPlatform.h"
+#include "Interfaces/ITargetPlatformModule.h"
+
+#define LOCTEXT_NAMESPACE "IOSTargetSettings"
+DEFINE_LOG_CATEGORY_STATIC(LogIOSTargetSettings, Log, All);
+
+bool SProvisionListRow::bInitialized = false;
+FCheckBoxStyle SProvisionListRow::ProvisionCheckBoxStyle;
+
+const FString gProjectNameText("[PROJECT_NAME]");
+
+//////////////////////////////////////////////////////////////////////////
+// FIOSTargetSettingsCustomization
+namespace FIOSTargetSettingsCustomizationConstants
+{
+	const FText DisabledTip = LOCTEXT("GitHubSourceRequiredToolTip", "This requires GitHub source.");
+}
+
+
+TSharedRef<IDetailCustomization> FIOSTargetSettingsCustomization::MakeInstance()
+{
+	return MakeShareable(new FIOSTargetSettingsCustomization);
+}
+
+FIOSTargetSettingsCustomization::FIOSTargetSettingsCustomization()
+	: EngineInfoPath(FString::Printf(TEXT("%sBuild/IOS/UE4Game-Info.plist"), *FPaths::EngineDir()))
+	, GameInfoPath(FString::Printf(TEXT("%sBuild/IOS/Info.plist"), *FPaths::GameDir()))
+	, EngineGraphicsPath(FString::Printf(TEXT("%sBuild/IOS/Resources/Graphics"), *FPaths::EngineDir()))
+	, GameGraphicsPath(FString::Printf(TEXT("%sBuild/IOS/Resources/Graphics"), *FPaths::GameDir()))
+{
+	new (IconNames) FPlatformIconInfo(TEXT("Icon29.png"), LOCTEXT("SettingsIcon_iPhone", "iPhone Settings Icon"), FText::GetEmpty(), 29, 29, FPlatformIconInfo::Optional);// also iOS6 spotlight search
+	new (IconNames) FPlatformIconInfo(TEXT("[REDACTED_EMAIL]"), LOCTEXT("SettingsIcon_iPhoneRetina", "iPhone Retina Settings Icon"), FText::GetEmpty(), 58, 58, FPlatformIconInfo::Optional); // also iOS6 spotlight search
+	new (IconNames) FPlatformIconInfo(TEXT("Icon40.png"), LOCTEXT("SpotlightIcon_iOS7", "iOS7 Spotlight Icon"), FText::GetEmpty(), 40, 40, FPlatformIconInfo::Optional);
+	new (IconNames) FPlatformIconInfo(TEXT("[REDACTED_EMAIL]"), LOCTEXT("SpotlightIcon_Retina_iOS7", "Retina iOS7 Spotlight Icon"), FText::GetEmpty(), 80, 80, FPlatformIconInfo::Optional);
+	new (IconNames) FPlatformIconInfo(TEXT("Icon50.png"), LOCTEXT("SpotlightIcon_iPad_iOS6", "iPad iOS6 Spotlight Icon"), FText::GetEmpty(), 50, 50, FPlatformIconInfo::Optional);
+	new (IconNames) FPlatformIconInfo(TEXT("[REDACTED_EMAIL]"), LOCTEXT("SpotlightIcon_iPadRetina_iOS6", "iPad Retina iOS6 Spotlight Icon"), FText::GetEmpty(), 100, 100, FPlatformIconInfo::Optional);
+	new (IconNames) FPlatformIconInfo(TEXT("Icon57.png"), LOCTEXT("AppIcon_iPhone_iOS6", "iPhone iOS6 App Icon"), FText::GetEmpty(), 57, 57, FPlatformIconInfo::Required);
+	new (IconNames) FPlatformIconInfo(TEXT("[REDACTED_EMAIL]"), LOCTEXT("AppIcon_iPhoneRetina_iOS6", "iPhone Retina iOS6 App Icon"), FText::GetEmpty(), 114, 114, FPlatformIconInfo::Required);
+	new (IconNames) FPlatformIconInfo(TEXT("[REDACTED_EMAIL]"), LOCTEXT("AppIcon_iPhoneRetina_iOS7", "iPhone Retina iOS7 App Icon"), FText::GetEmpty(), 120, 120, FPlatformIconInfo::Required);
+	new (IconNames)FPlatformIconInfo(TEXT("[REDACTED_EMAIL]"), LOCTEXT("AppIcon_iPhoneRetina_iOS8", "iPhone Plus Retina iOS8 App Icon"), FText::GetEmpty(), 180, 180, FPlatformIconInfo::Required);
+	new (IconNames)FPlatformIconInfo(TEXT("Icon72.png"), LOCTEXT("AppIcon_iPad_iOS6", "iPad iOS6 App Icon"), FText::GetEmpty(), 72, 72, FPlatformIconInfo::Required);
+	new (IconNames) FPlatformIconInfo(TEXT("[REDACTED_EMAIL]"), LOCTEXT("AppIcon_iPadRetina_iOS6", "iPad Retina iOS6 App Icon"), FText::GetEmpty(), 144, 144, FPlatformIconInfo::Required);
+	new (IconNames) FPlatformIconInfo(TEXT("Icon76.png"), LOCTEXT("AppIcon_iPad_iOS7", "iPad iOS7 App Icon"), FText::GetEmpty(), 76, 76, FPlatformIconInfo::Required);
+	new (IconNames) FPlatformIconInfo(TEXT("[REDACTED_EMAIL]"), LOCTEXT("AppIcon_iPadRetina_iOS7", "iPad Retina iOS7 App Icon"), FText::GetEmpty(), 152, 152, FPlatformIconInfo::Required);
+	new (IconNames)FPlatformIconInfo(TEXT("[REDACTED_EMAIL]"), LOCTEXT("AppIcon_iPadProRetina_iOS9", "iPad Pro Retina iOS9 App Icon"), FText::GetEmpty(), 167, 167, FPlatformIconInfo::Required);
+
+	new (LaunchImageNames) FPlatformIconInfo(TEXT("Default.png"), LOCTEXT("LaunchImage_iPhone", "Launch iPhone 4/4S"), FText::GetEmpty(), 320, 480, FPlatformIconInfo::Required);
+	new (LaunchImageNames) FPlatformIconInfo(TEXT("[REDACTED_EMAIL]"), LOCTEXT("LaunchImage_iPhoneRetina", "Launch iPhone 4/4S Retina"), FText::GetEmpty(), 640, 960, FPlatformIconInfo::Required);
+	new (LaunchImageNames) FPlatformIconInfo(TEXT("[REDACTED_EMAIL]"), LOCTEXT("LaunchImage_iPhone5", "Launch iPhone 5/5S Retina"), FText::GetEmpty(), 640, 1136, FPlatformIconInfo::Required);
+	new (LaunchImageNames) FPlatformIconInfo(TEXT("Default-Landscape.png"), LOCTEXT("LaunchImage_iPad_Landscape", "Launch iPad in Landscape"), FText::GetEmpty(), 1024, 768, FPlatformIconInfo::Required);
+	new (LaunchImageNames) FPlatformIconInfo(TEXT("[REDACTED_EMAIL]"), LOCTEXT("LaunchImage_iPadRetina_Landscape", "Launch iPad Retina in Landscape"), FText::GetEmpty(), 2048, 1536, FPlatformIconInfo::Required);
+	new (LaunchImageNames) FPlatformIconInfo(TEXT("Default-Portrait.png"), LOCTEXT("LaunchImage_iPad_Portrait", "Launch iPad in Portrait"), FText::GetEmpty(), 768, 1024, FPlatformIconInfo::Required);
+	new (LaunchImageNames) FPlatformIconInfo(TEXT("[REDACTED_EMAIL]"), LOCTEXT("LaunchImage_iPadRetina_Portrait", "Launch iPad Retina in Portrait"), FText::GetEmpty(), 1536, 2048, FPlatformIconInfo::Required);
+	new (LaunchImageNames) FPlatformIconInfo(TEXT("Default-IPhone6.png"), LOCTEXT("LaunchImage_iPhone6", "Launch iPhone 6 in Portrait"), FText::GetEmpty(), 750, 1334, FPlatformIconInfo::Required);
+	new (LaunchImageNames)FPlatformIconInfo(TEXT("Default-IPhone6-Landscape.png"), LOCTEXT("LaunchImage_iPhone6_Landscape", "Launch iPhone 6 in Landscape"), FText::GetEmpty(), 1334, 750, FPlatformIconInfo::Required);
+	new (LaunchImageNames)FPlatformIconInfo(TEXT("Default-IPhone6Plus-Landscape.png"), LOCTEXT("LaunchImage_iPhone6Plus_Landscape", "Launch iPhone 6 Plus in Landscape"), FText::GetEmpty(), 2208, 1242, FPlatformIconInfo::Required);
+	new (LaunchImageNames) FPlatformIconInfo(TEXT("Default-IPhone6Plus-Portrait.png"), LOCTEXT("LaunchImage_iPhone6Plus_Portrait", "Launch iPhone 6 Plus in Portrait"), FText::GetEmpty(), 1242, 2208, FPlatformIconInfo::Required);
+	new (LaunchImageNames) FPlatformIconInfo(TEXT("Default-Landscape-1336.png"), LOCTEXT("LaunchImage_iPadPro_Landscape", "Launch iPad Pro in Landscape"), FText::GetEmpty(), 1336, 1024, FPlatformIconInfo::Required);
+	new (LaunchImageNames)FPlatformIconInfo(TEXT("Default-Portrait-1336.png"), LOCTEXT("LaunchImage_iPadPro_Portrait", "Launch iPad Pro in Portrait"), FText::GetEmpty(), 1024, 1336, FPlatformIconInfo::Required);
+	new (LaunchImageNames)FPlatformIconInfo(TEXT("[REDACTED_EMAIL]"), LOCTEXT("LaunchImage_iPadProRetina_Landscape", "Launch iPad Pro Retina in Landscape"), FText::GetEmpty(), 2732, 2048, FPlatformIconInfo::Required);
+	new (LaunchImageNames)FPlatformIconInfo(TEXT("[REDACTED_EMAIL]"), LOCTEXT("LaunchImage_iPadProRetina_Portrait", "Launch iPad Pro Retina in Portrait"), FText::GetEmpty(), 2048, 2732, FPlatformIconInfo::Required);
+
+	bShowAllProvisions = false;
+	bShowAllCertificates = false;
+	ProvisionList = MakeShareable(new TArray<ProvisionPtr>());
+	CertificateList = MakeShareable(new TArray<CertificatePtr>());
+}
+
+FIOSTargetSettingsCustomization::~FIOSTargetSettingsCustomization()
+{
+	if (IPPProcess.IsValid())
+	{
+		IPPProcess = NULL;
+		FTicker::GetCoreTicker().RemoveTicker(TickerHandle);
+	}
+}
+
+void FIOSTargetSettingsCustomization::CustomizeDetails(IDetailLayoutBuilder& DetailLayout)
+{
+	SavedLayoutBuilder = &DetailLayout;
+
+	BuildPListSection(DetailLayout);
+
+	BuildIconSection(DetailLayout);
+
+	BuildRemoteBuildingSection(DetailLayout);
+
+	FindRequiredFiles();
+}
+
+static FString OutputMessage;
+static void OnOutput(FString Message)
+{
+	OutputMessage += Message;
+	OutputMessage += "\n";
+	UE_LOG(LogTemp, Log, TEXT("%s"), *Message);
+}
+
+void FIOSTargetSettingsCustomization::UpdateStatus()
+{
+	if (OutputMessage.Len() > 0)
+	{
+		CertificateList->Reset();
+		ProvisionList->Reset();
+
+		// Now split up the log into multiple lines
+		TArray<FString> LogLines;
+		OutputMessage.ParseIntoArray(LogLines, TEXT("\n"), true);
+		
+		// format of the line being read here!!
+		bool bCerts = false;
+		bManuallySelected = false;
+		for (int Index = 0; Index < LogLines.Num(); Index++)
+		{
+			FString& Line = LogLines[Index];
+			TArray<FString> Fields;
+			Line.ParseIntoArray(Fields, TEXT(","), true);
+			if (Line.Contains(TEXT("CERTIFICATE-"), ESearchCase::CaseSensitive))
+			{
+				CertificatePtr Cert = MakeShareable<FCertificate>(new FCertificate());
+				for (int FieldIndex = 0; FieldIndex < Fields.Num(); ++FieldIndex)
+				{
+					FString Key, Value;
+					Fields[FieldIndex].Split(TEXT(":"), &Key, &Value);
+					if (Key.Contains("Name"))
+					{
+						Cert->Name = Value;
+					}
+					else if (Key.Contains(TEXT("Validity")))
+					{
+						Cert->Status = Value;
+					}
+					else if (Key.Contains(TEXT("EndDate")))
+					{
+						FString Date, Time;
+						Value.Split(TEXT("T"), &Date, &Time);
+						Cert->Expires = Date;
+					}
+				}
+				CertificatePtr PrevCert = NULL;
+				for (int CIndex = 0; CIndex < CertificateList->Num() && !PrevCert.IsValid(); ++CIndex)
+				{
+					if ((*CertificateList)[CIndex]->Name == Cert->Name)
+					{
+						PrevCert = (*CertificateList)[CIndex];
+						break;
+					}
+				}
+
+				// check to see if this the one selected in the ini file
+				FString OutString;
+				SignCertificateProperty->GetValueAsFormattedString(OutString);
+				Cert->bManuallySelected = (OutString == Cert->Name);
+				bManuallySelected |= Cert->bManuallySelected;
+				if (!PrevCert.IsValid())
+				{
+					CertificateList->Add(Cert);
+				}
+				else
+				{
+					FDateTime time1, time2;
+					FDateTime::ParseIso8601(*(PrevCert->Expires), time1);
+					FDateTime::ParseIso8601(*(Cert->Expires), time2);
+					if (time2 > time1)
+					{
+						PrevCert->Expires = Cert->Expires;
+						PrevCert->Status = Cert->Status;
+					}
+					Cert = NULL;
+				}
+			}
+			else if (Line.Contains(TEXT("PROVISION-"), ESearchCase::CaseSensitive))
+			{
+				ProvisionPtr Prov = MakeShareable<FProvision>(new FProvision());
+				for (int FieldIndex = 0; FieldIndex < Fields.Num(); ++FieldIndex)
+				{
+					FString Key, Value;
+					Fields[FieldIndex].Split(TEXT(":"), &Key, &Value);
+					if (Key.Contains("File"))
+					{
+						Prov->FileName = Value;
+					}
+					else if (Key.Contains("Name"))
+					{
+						Prov->Name = Value;
+					}
+					else if (Key.Contains(TEXT("Validity")))
+					{
+						Prov->Status = Value;
+					}
+					else if (Key.Contains(TEXT("Type")))
+					{
+						Prov->bDistribution = Value.Contains(TEXT("DISTRIBUTION"));
+					}
+				}
+
+				// check to see if this the one selected in the ini file
+				FString OutString;
+				MobileProvisionProperty->GetValueAsFormattedString(OutString);
+				Prov->bManuallySelected = (OutString == Prov->FileName);
+				bManuallySelected |= Prov->bManuallySelected;
+				ProvisionList->Add(Prov);
+			}
+			else if (Line.Contains(TEXT("MATCHED-"), ESearchCase::CaseSensitive))
+			{
+				for (int FieldIndex = 0; FieldIndex < Fields.Num(); ++FieldIndex)
+				{
+					FString Key, Value;
+					Fields[FieldIndex].Split(TEXT(":"), &Key, &Value);
+					if (Key.Contains("File"))
+					{
+						SelectedFile = Value;
+					}
+					else if (Key.Contains("Provision"))
+					{
+						SelectedProvision = Value;
+					}
+					else if (Key.Contains(TEXT("Cert")))
+					{
+						SelectedCert = Value;
+					}
+				}
+			}
+		}
+
+		FilterLists();
+	}
+}
+
+void FIOSTargetSettingsCustomization::UpdateSSHStatus()
+{
+	// updated SSH key
+	const UIOSRuntimeSettings* Settings = GetDefault<UIOSRuntimeSettings>();
+	const_cast<UIOSRuntimeSettings*>(Settings)->PostInitProperties();
+}
+
+void FIOSTargetSettingsCustomization::BuildPListSection(IDetailLayoutBuilder& DetailLayout)
+{
+	// Info.plist category
+	IDetailCategoryBuilder& ProvisionCategory = DetailLayout.EditCategory(TEXT("Mobile Provision"));
+	IDetailCategoryBuilder& AppManifestCategory = DetailLayout.EditCategory(TEXT("Info.plist"));
+	IDetailCategoryBuilder& BundleCategory = DetailLayout.EditCategory(TEXT("BundleInformation"));
+	IDetailCategoryBuilder& OrientationCategory = DetailLayout.EditCategory(TEXT("Orientation"));
+	IDetailCategoryBuilder& RenderCategory = DetailLayout.EditCategory(TEXT("Rendering"));
+	IDetailCategoryBuilder& OSInfoCategory = DetailLayout.EditCategory(TEXT("OS Info"));
+	IDetailCategoryBuilder& DeviceCategory = DetailLayout.EditCategory(TEXT("Devices"));
+	IDetailCategoryBuilder& BuildCategory = DetailLayout.EditCategory(TEXT("Build"));
+	IDetailCategoryBuilder& OnlineCategory = DetailLayout.EditCategory(TEXT("Online"));
+	IDetailCategoryBuilder& ExtraCategory = DetailLayout.EditCategory(TEXT("Extra PList Data"));
+	MobileProvisionProperty = DetailLayout.GetProperty(GET_MEMBER_NAME_CHECKED(UIOSRuntimeSettings, MobileProvision));
+	BuildCategory.AddProperty(MobileProvisionProperty)
+		.Visibility(EVisibility::Hidden);
+	SignCertificateProperty = DetailLayout.GetProperty(GET_MEMBER_NAME_CHECKED(UIOSRuntimeSettings, SigningCertificate));
+	BuildCategory.AddProperty(SignCertificateProperty)
+		.Visibility(EVisibility::Hidden);
+	AutomaticSigningProperty = DetailLayout.GetProperty(GET_MEMBER_NAME_CHECKED(UIOSRuntimeSettings, bAutomaticSigning));
+	BuildCategory.AddProperty(AutomaticSigningProperty)
+		.Visibility(EVisibility::Hidden);
+	ProvisionCategory.AddProperty(AutomaticSigningProperty);
+
+/*	ProvisionCategory.AddCustomRow(TEXT("Certificate Request"), false)
+		.NameContent()
+		[
+			SNew(SHorizontalBox)
+			+ SHorizontalBox::Slot()
+			.Padding(FMargin(0, 1, 0, 1))
+			.FillWidth(1.0f)
+			[
+				SNew(STextBlock)
+				.Text(LOCTEXT("RequestLabel", "Certificate Request"))
+				.Font(DetailLayout.GetDetailFont())
+			]
+		]
+	.ValueContent()
+		[
+			SNew(SBox)
+			.HAlign(HAlign_Center)
+			[
+				SNew(SButton)
+				.HAlign(HAlign_Center)
+				.OnClicked(this, &FIOSTargetSettingsCustomization::OnCertificateRequestClicked)
+				[
+					SNew(STextBlock)
+					.Text(FText::FromString("Create Certificate Request and a Key Pair"))
+				]
+			]
+		];*/
+
+	ProvisionCategory.AddCustomRow(LOCTEXT("ProvisionLabel", "Provision"), false)
+		.WholeRowWidget
+		.MinDesiredWidth(0.f)
+		.MaxDesiredWidth(0.f)
+		.HAlign(HAlign_Fill)
+		[
+			SNew(SVerticalBox)
+			+ SVerticalBox::Slot()
+			.AutoHeight()
+			[
+				SAssignNew(ProvisionInfoSwitcher, SWidgetSwitcher)
+				.WidgetIndex(0)
+				// searching for provisions
+				+SWidgetSwitcher::Slot()
+				.HAlign(HAlign_Fill)
+				.VAlign(VAlign_Center)
+				[
+					SNew(SBorder)
+					.Padding(4)
+					[
+						SNew( STextBlock )
+						.Text( LOCTEXT( "ProvisionViewerFindingProvisions", "Please wait while we gather information." ) )
+						.AutoWrapText( true )
+					]
+				]
+				// importing a provision
+				+SWidgetSwitcher::Slot()
+					.HAlign(HAlign_Fill)
+					.VAlign(VAlign_Center)
+					[
+						SNew(SBorder)
+						.Padding(4)
+						[
+							SNew( STextBlock )
+							.Text( LOCTEXT( "ProvisionViewerImportingProvisions", "Importing Provision.  Please wait..." ) )
+							.AutoWrapText( true )
+						]
+					]
+				// no provisions found or no valid provisions
+				+SWidgetSwitcher::Slot()
+				.HAlign(HAlign_Fill)
+				.VAlign(VAlign_Center)
+				[
+					SNew(SBorder)
+					.Padding(4)
+					[
+						SNew( STextBlock )
+						.Text( LOCTEXT( "ProvisionViewerNoValidProvisions", "No Provisions Found. Please Import a Provision." ) )
+						.AutoWrapText( true )
+					]
+				]
+				+SWidgetSwitcher::Slot()
+				.HAlign(HAlign_Fill)
+				.VAlign(VAlign_Center)
+				[
+					SNew(SVerticalBox)
+					+ SVerticalBox::Slot()
+					.Padding(FMargin(10, 10, 10, 10))
+					.AutoHeight()
+					[
+						SAssignNew(ProvisionListView, SListView<ProvisionPtr>)
+						.ItemHeight(20.0f)
+						.ListItemsSource(&FilteredProvisionList)
+						.OnGenerateRow(this, &FIOSTargetSettingsCustomization::HandleProvisionListGenerateRow)
+						.SelectionMode(ESelectionMode::None)
+						.HeaderRow
+						(
+						SNew(SHeaderRow)
+						+ SHeaderRow::Column("Selected")
+						.DefaultLabel(LOCTEXT("ProvisionListSelectColumnHeader", ""))
+						.FixedWidth(30.0f)
+						+ SHeaderRow::Column("Name")
+						.DefaultLabel(LOCTEXT("ProvisionListNameColumnHeader", "Provision"))
+						.FillWidth(1.0f)
+						+ SHeaderRow::Column("File")
+						.DefaultLabel(LOCTEXT("ProvisionListFileColumnHeader", "File"))
+						+ SHeaderRow::Column("Status")
+						.DefaultLabel(LOCTEXT("ProvisionListStatusColumnHeader", "Status"))
+						+ SHeaderRow::Column("Distribution")
+						.DefaultLabel(LOCTEXT("ProvisionListDistributionColumnHeader", "Distribution"))
+						.FixedWidth(75.0f)
+						)
+					]
+					+ SVerticalBox::Slot()
+						.AutoHeight()
+						.Padding(0.0f, 6.0f, 0.0f, 4.0f)
+						[
+							SNew(SSeparator)
+							.Orientation(Orient_Horizontal)
+						]
+					+ SVerticalBox::Slot()
+						.AutoHeight()
+						[
+							SNew(SHorizontalBox)
+
+							+ SHorizontalBox::Slot()
+							.AutoWidth()
+							[
+								SNew(SRichTextBlock)
+								.Text(LOCTEXT("ProvisionMessage", "<RichTextBlock.TextHighlight>Note</>: If no provision is selected the one in green will be used to provision the IPA."))
+								.TextStyle(FEditorStyle::Get(), "MessageLog")
+								.DecoratorStyleSet(&FEditorStyle::Get())
+								.AutoWrapText(true)
+							]
+
+							+ SHorizontalBox::Slot()
+								.FillWidth(1.0f)
+								.HAlign(HAlign_Right)
+								[
+									SNew(STextBlock)
+									.Text(LOCTEXT("ViewLabel", "View:"))
+								]
+
+							+ SHorizontalBox::Slot()
+								.AutoWidth()
+								.Padding(8.0f, 0.0f)
+								[
+									// all provisions hyper link
+									SNew(SHyperlink)
+									.OnNavigate(this, &FIOSTargetSettingsCustomization::HandleAllProvisionsHyperlinkNavigate, true)
+									.Text(LOCTEXT("AllProvisionsHyperLinkLabel", "All"))
+									.ToolTipText(LOCTEXT("AllProvisionsButtonTooltip", "View all provisions."))
+								]
+
+							+ SHorizontalBox::Slot()
+								.AutoWidth()
+								[
+									// valid provisions hyper link
+									SNew(SHyperlink)
+									.OnNavigate(this, &FIOSTargetSettingsCustomization::HandleAllProvisionsHyperlinkNavigate, false)
+									.Text(LOCTEXT("ValidProvisionsHyperlinkLabel", "Valid Only"))
+									.ToolTipText(LOCTEXT("ValidProvisionsHyperlinkTooltip", "View Valid provisions."))
+								]
+						]
+				]
+			]
+			+ SVerticalBox::Slot()
+			.AutoHeight()
+			[
+				SNew(SHorizontalBox)
+				+ SHorizontalBox::Slot()
+					.Padding(FMargin(0, 5, 0, 10))
+					.AutoWidth()
+					[
+						SNew(SButton)
+						.HAlign(HAlign_Center)
+						.VAlign(VAlign_Center)
+						.OnClicked(this, &FIOSTargetSettingsCustomization::OnInstallProvisionClicked)
+						.IsEnabled(this, &FIOSTargetSettingsCustomization::IsImportEnabled)
+						[
+							SNew(STextBlock)
+							.Text(LOCTEXT("ImportProvision", "Import Provision"))
+						]
+					]
+			]
+		];
+
+	ProvisionCategory.AddCustomRow(LOCTEXT("CertificateLabel", "Certificate"), false)
+		.WholeRowWidget
+		.MinDesiredWidth(0.f)
+		.MaxDesiredWidth(0.f)
+		.HAlign(HAlign_Fill)
+		[
+			SNew(SVerticalBox)
+			+ SVerticalBox::Slot()
+			.AutoHeight()
+			[
+				SAssignNew(CertificateInfoSwitcher, SWidgetSwitcher)
+				.WidgetIndex(0)
+				// searching for provisions
+				+SWidgetSwitcher::Slot()
+				.HAlign(HAlign_Fill)
+				.VAlign(VAlign_Center)
+				[
+					SNew(SBorder)
+					.Padding(4)
+					[
+						SNew( STextBlock )
+						.Text( LOCTEXT( "CertificateViewerFindingProvisions", "Please wait while we gather information." ) )
+						.AutoWrapText( true )
+					]
+				]
+				// importing certificate
+				+SWidgetSwitcher::Slot()
+					.HAlign(HAlign_Fill)
+					.VAlign(VAlign_Center)
+					[
+						SNew(SBorder)
+						.Padding(4)
+						[
+							SNew( STextBlock )
+							.Text( LOCTEXT( "CertificateViewerImportingCertificate", "Importing Certificate.  Please wait..." ) )
+							.AutoWrapText( true )
+						]
+					]
+				// no provisions found or no valid provisions
+				+SWidgetSwitcher::Slot()
+					.HAlign(HAlign_Fill)
+					.VAlign(VAlign_Center)
+					[
+						SNew(SBorder)
+						.Padding(4)
+						[
+							SNew( STextBlock )
+							.Text( LOCTEXT( "CertificateViewerNoValidProvisions", "No Certificates Found.  Please Import a Certificate." ) )
+							.AutoWrapText( true )
+						]
+					]
+				+SWidgetSwitcher::Slot()
+					.HAlign(HAlign_Fill)
+					.VAlign(VAlign_Center)
+					[
+						SNew(SVerticalBox)
+						+ SVerticalBox::Slot()
+						.AutoHeight()
+						[
+							SNew(SHorizontalBox)
+							+ SHorizontalBox::Slot()
+							.Padding(FMargin(10, 10, 10, 10))
+							.FillWidth(1.0f)
+							[
+								SAssignNew(CertificateListView, SListView<CertificatePtr>)
+								.ItemHeight(20.0f)
+								.ListItemsSource(&FilteredCertificateList)
+								.OnGenerateRow(this, &FIOSTargetSettingsCustomization::HandleCertificateListGenerateRow)
+								.SelectionMode(ESelectionMode::None)
+								.HeaderRow
+								(
+									SNew(SHeaderRow)
+									+ SHeaderRow::Column("Selected")
+									.DefaultLabel(LOCTEXT("CertificateListSelectColumnHeader", ""))
+									.FixedWidth(30.0f)
+									+ SHeaderRow::Column("Name")
+									.DefaultLabel(LOCTEXT("CertificateListNameColumnHeader", "Certificate"))
+									+ SHeaderRow::Column("Status")
+									.DefaultLabel(LOCTEXT("CertificateListStatusColumnHeader", "Status"))
+									.FixedWidth(75.0f)
+									+ SHeaderRow::Column("Expires")
+									.DefaultLabel(LOCTEXT("CertificateListExpiresColumnHeader", "Expires"))
+									.FixedWidth(75.0f)
+								)
+							]
+						]
+						+ SVerticalBox::Slot()
+							.AutoHeight()
+							.Padding(0.0f, 6.0f, 0.0f, 4.0f)
+							[
+								SNew(SSeparator)
+								.Orientation(Orient_Horizontal)
+							]
+						+ SVerticalBox::Slot()
+							.AutoHeight()
+							[
+								SNew(SHorizontalBox)
+								+ SHorizontalBox::Slot()
+								.AutoWidth()
+								[
+									SNew(SRichTextBlock)
+									.Text(LOCTEXT("CertificateMessage", "<RichTextBlock.TextHighlight>Note</>: If no certificate is selected then the one in green will be used to sign the IPA."))
+									.TextStyle(FEditorStyle::Get(), "MessageLog")
+									.DecoratorStyleSet(&FEditorStyle::Get())
+									.AutoWrapText(true)
+								]
+								+ SHorizontalBox::Slot()
+								.FillWidth(1.0f)
+								.HAlign(HAlign_Right)
+								[
+									SNew(STextBlock)
+									.Text(LOCTEXT("ViewLabel", "View:"))
+								]
+
+								+ SHorizontalBox::Slot()
+									.AutoWidth()
+									.Padding(8.0f, 0.0f)
+									[
+										// all provisions hyper link
+										SNew(SHyperlink)
+										.OnNavigate(this, &FIOSTargetSettingsCustomization::HandleAllCertificatesHyperlinkNavigate, true)
+										.Text(LOCTEXT("AllCertificatesHyperLinkLabel", "All"))
+										.ToolTipText(LOCTEXT("AllCertificatesButtonTooltip", "View all certificates."))
+									]
+
+								+ SHorizontalBox::Slot()
+									.AutoWidth()
+									[
+										// valid provisions hyper link
+										SNew(SHyperlink)
+										.OnNavigate(this, &FIOSTargetSettingsCustomization::HandleAllCertificatesHyperlinkNavigate, false)
+										.Text(LOCTEXT("ValidCertificatesHyperlinkLabel", "Valid Only"))
+										.ToolTipText(LOCTEXT("ValidCertificatesHyperlinkTooltip", "View Valid certificates."))
+									]
+							]
+					]
+				]
+			+ SVerticalBox::Slot()
+			.AutoHeight()
+			[
+				SNew(SHorizontalBox)
+				+ SHorizontalBox::Slot()
+				.Padding(FMargin(0, 5, 0, 10))
+				.AutoWidth()
+				[
+					SNew(SButton)
+					.HAlign(HAlign_Center)
+					.VAlign(VAlign_Center)
+					.OnClicked(this, &FIOSTargetSettingsCustomization::OnInstallCertificateClicked)
+					.IsEnabled(this, &FIOSTargetSettingsCustomization::IsImportEnabled)
+					[
+						SNew(STextBlock)
+						.Text(LOCTEXT("ImportCertificate", "Import Certificate"))
+					]
+				]
+			]
+		];
+	
+	BundleCategory.AddCustomRow(LOCTEXT("UpgradeInfo", "Upgrade Info"), false)
+	.WholeRowWidget
+	[
+		SNew(SBorder)
+		.Padding(1)
+		[
+			SNew(SHorizontalBox)
+			+ SHorizontalBox::Slot()
+			.Padding(FMargin(10, 10, 10, 10))
+			.FillWidth(1.0f)
+			[
+				SNew(SRichTextBlock)
+				.Text(LOCTEXT("IOSUpgradeInfoMessage", "<RichTextBlock.TextHighlight>Note to users from 4.6 or earlier</>: We now <RichTextBlock.TextHighlight>GENERATE</> an Info.plist when building, so if you have customized your .plist file, you will need to put all of your changes into the below settings. Note that we don't touch the .plist file that is in your project directory, so you can use it as reference."))
+				.TextStyle(FEditorStyle::Get(), "MessageLog")
+				.DecoratorStyleSet(&FEditorStyle::Get())
+				.AutoWrapText(true)
+				// + SRichTextBlock::HyperlinkDecorator(TEXT("browser"), FSlateHyperlinkRun::FOnClick::CreateStatic(&OnBrowserLinkClicked))
+			 ]
+		 ]
+	 ];
+
+	// Show properties that are gated by the plist being present and writable
+	RunningIPPProcess = false;
+
+#define SETUP_SOURCEONLY_PROP(PropName, Category) \
+	{ \
+		TSharedRef<IPropertyHandle> PropertyHandle = DetailLayout.GetProperty(GET_MEMBER_NAME_CHECKED(UIOSRuntimeSettings, PropName)); \
+		Category.AddProperty(PropertyHandle) \
+			.IsEnabled(FEngineBuildSettings::IsSourceDistribution()) \
+			.ToolTip(FEngineBuildSettings::IsSourceDistribution() ? PropertyHandle->GetToolTipText() : FIOSTargetSettingsCustomizationConstants::DisabledTip); \
+	}
+
+#define SETUP_PLIST_PROP(PropName, Category) \
+	{ \
+		TSharedRef<IPropertyHandle> PropertyHandle = DetailLayout.GetProperty(GET_MEMBER_NAME_CHECKED(UIOSRuntimeSettings, PropName)); \
+		Category.AddProperty(PropertyHandle); \
+	}
+
+#define SETUP_STATUS_PROP(PropName, Category) \
+	{ \
+		TSharedRef<IPropertyHandle> PropertyHandle = DetailLayout.GetProperty(GET_MEMBER_NAME_CHECKED(UIOSRuntimeSettings, PropName)); \
+		Category.AddProperty(PropertyHandle) \
+		.Visibility(EVisibility::Hidden); \
+		Category.AddCustomRow(LOCTEXT("BundleIdentifier", "BundleIdentifier"), false) \
+		.NameContent() \
+		[ \
+			SNew(SHorizontalBox) \
+			+ SHorizontalBox::Slot() \
+			.Padding(FMargin(0, 1, 0, 1)) \
+			.FillWidth(1.0f) \
+			[ \
+				SNew(STextBlock) \
+				.Text(LOCTEXT("BundleIdentifierLabel", "Bundle Identifier")) \
+				.Font(DetailLayout.GetDetailFont()) \
+			]\
+		] \
+		.ValueContent() \
+		.MinDesiredWidth( 0.0f ) \
+		.MaxDesiredWidth( 0.0f ) \
+		[ \
+			SNew(SHorizontalBox) \
+			+ SHorizontalBox::Slot() \
+			.FillWidth(1.0f) \
+			.HAlign(HAlign_Fill) \
+			[ \
+				SAssignNew(BundleIdTextBox, SEditableTextBox) \
+				.IsEnabled(this, &FIOSTargetSettingsCustomization::IsImportEnabled) \
+				.Text(this, &FIOSTargetSettingsCustomization::GetBundleText, PropertyHandle) \
+				.Font(DetailLayout.GetDetailFont()) \
+				.SelectAllTextOnCommit( true ) \
+				.SelectAllTextWhenFocused( true ) \
+				.ClearKeyboardFocusOnCommit(false) \
+				.ToolTipText(PropertyHandle->GetToolTipText()) \
+				.OnTextCommitted(this, &FIOSTargetSettingsCustomization::OnBundleIdentifierChanged, PropertyHandle) \
+				.OnTextChanged(this, &FIOSTargetSettingsCustomization::OnBundleIdentifierTextChanged, ETextCommit::Default, PropertyHandle) \
+			] \
+		]; \
+	}
+
+	const UIOSRuntimeSettings& Settings = *GetDefault<UIOSRuntimeSettings>();
+
+	FSimpleDelegate OnUpdateShaderStandardWarning = FSimpleDelegate::CreateSP(this, &FIOSTargetSettingsCustomization::UpdateShaderStandardWarning);
+	FSimpleDelegate OnUpdateOSVersionWarning = FSimpleDelegate::CreateSP(this, &FIOSTargetSettingsCustomization::UpdateOSVersionWarning);
+
+/*	GLES2PropertyHandle = DetailLayout.GetProperty(GET_MEMBER_NAME_CHECKED(UIOSRuntimeSettings, bSupportsOpenGLES2));
+	GLES2PropertyHandle->SetOnPropertyValueChanged(OnUpdateShaderStandardWarning);
+	RenderCategory.AddProperty(GLES2PropertyHandle);*/
+
+	MinOSPropertyHandle = DetailLayout.GetProperty(GET_MEMBER_NAME_CHECKED(UIOSRuntimeSettings, MinimumiOSVersion));
+	MinOSPropertyHandle->SetOnPropertyValueChanged(OnUpdateShaderStandardWarning);
+	OSInfoCategory.AddProperty(MinOSPropertyHandle);
+
+	DevArmV7PropertyHandle = DetailLayout.GetProperty(GET_MEMBER_NAME_CHECKED(UIOSRuntimeSettings, bDevForArmV7));
+	DevArmV7PropertyHandle->SetOnPropertyValueChanged(OnUpdateOSVersionWarning);
+	BuildCategory.AddProperty(DevArmV7PropertyHandle);
+
+	DevArmV7sPropertyHandle = DetailLayout.GetProperty(GET_MEMBER_NAME_CHECKED(UIOSRuntimeSettings, bDevForArmV7S));
+	DevArmV7sPropertyHandle->SetOnPropertyValueChanged(OnUpdateOSVersionWarning);
+	BuildCategory.AddProperty(DevArmV7sPropertyHandle);
+
+	ShipArmV7PropertyHandle = DetailLayout.GetProperty(GET_MEMBER_NAME_CHECKED(UIOSRuntimeSettings, bShipForArmV7));
+	ShipArmV7PropertyHandle->SetOnPropertyValueChanged(OnUpdateOSVersionWarning);
+	BuildCategory.AddProperty(ShipArmV7PropertyHandle);
+
+	ShipArmV7sPropertyHandle = DetailLayout.GetProperty(GET_MEMBER_NAME_CHECKED(UIOSRuntimeSettings, bShipForArmV7S));
+	ShipArmV7sPropertyHandle->SetOnPropertyValueChanged(OnUpdateOSVersionWarning);
+	BuildCategory.AddProperty(ShipArmV7sPropertyHandle);
+
+	SETUP_PLIST_PROP(BundleDisplayName, BundleCategory);
+	SETUP_PLIST_PROP(BundleName, BundleCategory);
+	SETUP_STATUS_PROP(BundleIdentifier, BundleCategory);
+	SETUP_PLIST_PROP(VersionInfo, BundleCategory);
+	SETUP_PLIST_PROP(bSupportsPortraitOrientation, OrientationCategory);
+	SETUP_PLIST_PROP(bSupportsUpsideDownOrientation, OrientationCategory);
+	SETUP_PLIST_PROP(bSupportsLandscapeLeftOrientation, OrientationCategory);
+	SETUP_PLIST_PROP(bSupportsLandscapeRightOrientation, OrientationCategory);
+	
+	SETUP_PLIST_PROP(bSupportsMetal, RenderCategory);
+
+	SETUP_SOURCEONLY_PROP(bEnableRemoteNotificationsSupport, OnlineCategory)
+	
+    // Handle max. shader version a little specially.
+    {
+        ShaderVersionPropertyHandle = DetailLayout.GetProperty(GET_MEMBER_NAME_CHECKED(UIOSRuntimeSettings, MaxShaderLanguageVersion));
+		
+		// Drop-downs for setting type of lower and upper bound normalization
+		IDetailPropertyRow& ShaderVersionPropertyRow = RenderCategory.AddProperty(ShaderVersionPropertyHandle.ToSharedRef());
+		ShaderVersionPropertyRow.CustomWidget()
+		.NameContent()
+		[
+			ShaderVersionPropertyHandle->CreatePropertyNameWidget()
+		 ]
+		.ValueContent()
+		.HAlign(HAlign_Fill)
+		[
+			SNew(SVerticalBox)
+			+ SVerticalBox::Slot()
+			.AutoHeight()
+			.Padding(2)
+			[
+				SNew(SComboButton)
+				.OnGetMenuContent(this, &FIOSTargetSettingsCustomization::OnGetShaderVersionContent)
+				.ContentPadding(FMargin( 2.0f, 2.0f ))
+				.ButtonContent()
+				[
+					SNew(STextBlock)
+					.Text(this, &FIOSTargetSettingsCustomization::GetShaderVersionDesc)
+					.Font(IDetailLayoutBuilder::GetDetailFont())
+				]
+			 ]
+			+ SVerticalBox::Slot()
+			.AutoHeight()
+			.HAlign(HAlign_Fill)
+			.Padding(2)
+			[
+				SAssignNew(ShaderVersionWarningTextBox, SErrorText)
+				.AutoWrapText(true)
+			]
+		];
+		
+		UpdateShaderStandardWarning();
+    }
+
+	SETUP_PLIST_PROP(bSupportsIPad, DeviceCategory);
+	SETUP_PLIST_PROP(bSupportsIPhone, DeviceCategory);
+	SETUP_PLIST_PROP(AdditionalPlistData, ExtraCategory);
+
+#undef SETUP_SOURCEONLY_PROP
+}
+
+
+void FIOSTargetSettingsCustomization::BuildRemoteBuildingSection(IDetailLayoutBuilder& DetailLayout)
+{
+#if PLATFORM_WINDOWS
+	IDetailCategoryBuilder& BuildCategory = DetailLayout.EditCategory(TEXT("Build"));
+
+	// Sub group we wish to add remote building options to.
+	FText RemoteBuildingGroupName = LOCTEXT("RemoteBuildingGroupName", "Remote Build Options");
+	IDetailGroup& RemoteBuildingGroup = BuildCategory.AddGroup(*RemoteBuildingGroupName.ToString(), RemoteBuildingGroupName, false);
+
+
+	// Remote Server Name Property
+	TSharedRef<IPropertyHandle> RemoteServerNamePropertyHandle = DetailLayout.GetProperty(GET_MEMBER_NAME_CHECKED(UIOSRuntimeSettings, RemoteServerName));
+	IDetailPropertyRow& RemoteServerNamePropertyRow = RemoteBuildingGroup.AddPropertyRow(RemoteServerNamePropertyHandle);
+	RemoteServerNamePropertyRow
+		.ToolTip(LOCTEXT("RemoteServerNameToolTip", "The name or ip address of the remote mac which will be used to build IOS"))
+		.CustomWidget()
+		.NameContent()
+		[
+			SNew(SHorizontalBox)
+			+ SHorizontalBox::Slot()
+			.Padding(FMargin(0, 1, 0, 1))
+			.FillWidth(1.0f)
+			[
+				SNew(STextBlock)
+				.Text(LOCTEXT("RemoteServerNameLabel", "Remote Server Name"))
+				.Font(DetailLayout.GetDetailFont())
+			]
+		]
+		.ValueContent()
+		.MinDesiredWidth(150.0f)
+		[
+			SNew(SHorizontalBox)
+			+ SHorizontalBox::Slot()
+			.Padding(FMargin(0.0f, 8.0f))
+			[
+				SNew(SEditableTextBox)
+				.IsEnabled(this, &FIOSTargetSettingsCustomization::IsImportEnabled)
+				.Text(this, &FIOSTargetSettingsCustomization::GetBundleText, RemoteServerNamePropertyHandle)
+				.Font(DetailLayout.GetDetailFont())
+				.SelectAllTextOnCommit(true)
+				.SelectAllTextWhenFocused(true)
+				.ClearKeyboardFocusOnCommit(false)
+				.ToolTipText(RemoteServerNamePropertyHandle->GetToolTipText())
+				.OnTextCommitted(this, &FIOSTargetSettingsCustomization::OnRemoteServerChanged, RemoteServerNamePropertyHandle)
+			]
+
+		];
+
+
+
+	// Add Use RSync Property
+	TSharedRef<IPropertyHandle> UseRSyncPropertyHandle = DetailLayout.GetProperty(GET_MEMBER_NAME_CHECKED(UIOSRuntimeSettings, bUseRSync));
+	BuildCategory.AddProperty(UseRSyncPropertyHandle)
+		.Visibility(EVisibility::Hidden);
+
+	// Add RSync Username Property
+	TSharedRef<IPropertyHandle> RSyncUsernamePropertyHandle = DetailLayout.GetProperty(GET_MEMBER_NAME_CHECKED(UIOSRuntimeSettings, RSyncUsername));
+	IDetailPropertyRow& RSyncUsernamePropertyRow = RemoteBuildingGroup.AddPropertyRow(RSyncUsernamePropertyHandle);
+	RSyncUsernamePropertyRow
+		.ToolTip(LOCTEXT("RSyncUsernameToolTip", "The username of the mac user that matches the specified SSH Key."))
+		.CustomWidget()
+			.NameContent()
+			[
+				SNew(SHorizontalBox)
+				+ SHorizontalBox::Slot()
+				.Padding(FMargin(0, 1, 0, 1))
+				.FillWidth(1.0f)
+				[
+					SNew(STextBlock)
+					.Text(LOCTEXT("RSyncUserNameLabel", "RSync User Name"))
+					.Font(DetailLayout.GetDetailFont())
+				]
+			]
+			.ValueContent()
+			.MinDesiredWidth(150.0f)
+			[
+				SNew(SHorizontalBox)
+				+ SHorizontalBox::Slot()
+				.Padding(FMargin(0.0f, 8.0f))
+				[
+					SNew(SEditableTextBox)
+					.IsEnabled(this, &FIOSTargetSettingsCustomization::IsImportEnabled)
+					.Text(this, &FIOSTargetSettingsCustomization::GetBundleText, RSyncUsernamePropertyHandle)
+					.Font(DetailLayout.GetDetailFont())
+					.SelectAllTextOnCommit(true)
+					.SelectAllTextWhenFocused(true)
+					.ClearKeyboardFocusOnCommit(false)
+					.ToolTipText(RSyncUsernamePropertyHandle->GetToolTipText())
+					.OnTextCommitted(this, &FIOSTargetSettingsCustomization::OnRemoteServerChanged, RSyncUsernamePropertyHandle)
+				]
+			];
+
+
+	// Add existing SSH path label.
+	TSharedRef<IPropertyHandle> SSHPrivateKeyLocationPropertyHandle = DetailLayout.GetProperty(GET_MEMBER_NAME_CHECKED(UIOSRuntimeSettings, SSHPrivateKeyLocation));
+	IDetailPropertyRow& SSHPrivateKeyLocationPropertyRow = RemoteBuildingGroup.AddPropertyRow(SSHPrivateKeyLocationPropertyHandle);
+	SSHPrivateKeyLocationPropertyRow
+		.ToolTip(LOCTEXT("SSHPrivateKeyLocationToolTip", "The existing location of an SSH Key found by UE4."));
+
+
+	// Add SSH override path
+	TSharedRef<IPropertyHandle> SSHPrivateKeyOverridePathPropertyHandle = DetailLayout.GetProperty(GET_MEMBER_NAME_CHECKED(UIOSRuntimeSettings, SSHPrivateKeyOverridePath));
+	IDetailPropertyRow& SSHPrivateKeyOverridePathPropertyRow = RemoteBuildingGroup.AddPropertyRow(SSHPrivateKeyOverridePathPropertyHandle);
+	SSHPrivateKeyOverridePathPropertyRow
+		.ToolTip(LOCTEXT("SSHPrivateKeyOverridePathToolTip", "Override the existing SSH Private Key with one from a specified location."));
+
+	// delta copy path
+	TSharedRef<IPropertyHandle> DeltaCopyOverridePathPropertyHandle = DetailLayout.GetProperty(GET_MEMBER_NAME_CHECKED(UIOSRuntimeSettings, DeltaCopyInstallPath));
+	IDetailPropertyRow& DeltaCopyOverridePathPropertyRow = RemoteBuildingGroup.AddPropertyRow(DeltaCopyOverridePathPropertyHandle);
+
+	const FText GenerateSSHText = LOCTEXT("GenerateSSHKey", "Generate SSH Key");
+
+	// Add a generate key button
+	RemoteBuildingGroup.AddWidgetRow()
+		.FilterString(GenerateSSHText)
+		.WholeRowWidget
+		.MinDesiredWidth(0.f)
+		.MaxDesiredWidth(0.f)
+		.HAlign(HAlign_Fill)
+		[
+			SNew(SVerticalBox)
+			+ SVerticalBox::Slot()
+				.AutoHeight()
+				[
+					SNew(SHorizontalBox)
+					+ SHorizontalBox::Slot()
+					.Padding(FMargin(0, 5, 0, 10))
+					.AutoWidth()
+					[
+						SNew(SButton)
+						.HAlign(HAlign_Center)
+						.VAlign(VAlign_Center)
+						.OnClicked(this, &FIOSTargetSettingsCustomization::OnGenerateSSHKey)
+						.IsEnabled(this, &FIOSTargetSettingsCustomization::IsImportEnabled)
+						[
+							SNew(STextBlock)
+							.Text(GenerateSSHText)
+						]
+					]
+				]
+		];
+#endif
+}
+
+
+void FIOSTargetSettingsCustomization::BuildIconSection(IDetailLayoutBuilder& DetailLayout)
+{
+	IDetailCategoryBuilder& RequiredIconCategory = DetailLayout.EditCategory(TEXT("Required Icons"));
+	IDetailCategoryBuilder& OptionalIconCategory = DetailLayout.EditCategory(TEXT("Optional Icons"));
+
+	// Add the icons
+	for (const FPlatformIconInfo& Info : IconNames)
+	{
+		const FVector2D IconImageMaxSize(Info.IconRequiredSize);
+		IDetailCategoryBuilder& IconCategory = (Info.RequiredState == FPlatformIconInfo::Required) ? RequiredIconCategory : OptionalIconCategory;
+		BuildImageRow(DetailLayout, IconCategory, Info, IconImageMaxSize);
+	}
+
+	// Add the launch images
+	IDetailCategoryBuilder& LaunchImageCategory = DetailLayout.EditCategory(TEXT("Launch Images"));
+	const FVector2D LaunchImageMaxSize(150.0f, 150.0f);
+	for (const FPlatformIconInfo& Info : LaunchImageNames)
+	{
+		BuildImageRow(DetailLayout, LaunchImageCategory, Info, LaunchImageMaxSize);
+	}
+}
+
+
+FReply FIOSTargetSettingsCustomization::OpenPlistFolder()
+{
+	const FString EditPlistFolder = FPaths::ConvertRelativePathToFull(FPaths::GetPath(GameInfoPath));
+	FPlatformProcess::ExploreFolder(*EditPlistFolder);
+
+	return FReply::Handled();
+}
+
+void FIOSTargetSettingsCustomization::CopySetupFilesIntoProject()
+{
+	// First copy the plist, it must get copied
+	FText ErrorMessage;
+	if (!SourceControlHelpers::CopyFileUnderSourceControl(GameInfoPath, EngineInfoPath, LOCTEXT("InfoPlist", "Info.plist"), /*out*/ ErrorMessage))
+	{
+		FNotificationInfo Info(ErrorMessage);
+		Info.ExpireDuration = 3.0f;
+		FSlateNotificationManager::Get().AddNotification(Info);
+	}
+	else
+	{
+		// Now try to copy all of the icons, etc... (these can be ignored if the file already exists)
+		TArray<FPlatformIconInfo> Graphics;
+		Graphics.Empty(IconNames.Num() + LaunchImageNames.Num());
+		Graphics.Append(IconNames);
+		Graphics.Append(LaunchImageNames);
+
+		for (const FPlatformIconInfo& Info : Graphics)
+		{
+			const FString EngineImagePath = EngineGraphicsPath / Info.IconPath;
+			const FString ProjectImagePath = GameGraphicsPath / Info.IconPath;
+
+			if (!FPaths::FileExists(ProjectImagePath))
+			{
+				SourceControlHelpers::CopyFileUnderSourceControl(ProjectImagePath, EngineImagePath, Info.IconName, /*out*/ ErrorMessage);
+			}
+		}
+	}
+
+	SavedLayoutBuilder->ForceRefreshDetails();
+}
+
+void FIOSTargetSettingsCustomization::BuildImageRow(IDetailLayoutBuilder& DetailLayout, IDetailCategoryBuilder& Category, const FPlatformIconInfo& Info, const FVector2D& MaxDisplaySize)
+{
+	const FString AutomaticImagePath = EngineGraphicsPath / Info.IconPath;
+	const FString TargetImagePath = GameGraphicsPath / Info.IconPath;
+
+	Category.AddCustomRow(Info.IconName)
+		.NameContent()
+		[
+			SNew(SHorizontalBox)
+			+ SHorizontalBox::Slot()
+			.Padding(FMargin(0, 1, 0, 1))
+			.FillWidth(1.0f)
+			[
+				SNew(STextBlock)
+				.Text(Info.IconName)
+				.Font(DetailLayout.GetDetailFont())
+			]
+		]
+		.ValueContent()
+		.MaxDesiredWidth(400.0f)
+		.MinDesiredWidth(100.0f)
+		[
+			SNew(SHorizontalBox)
+			+ SHorizontalBox::Slot()
+			.FillWidth(1.0f)
+			.VAlign(VAlign_Center)
+			[
+				SNew(SExternalImageReference, AutomaticImagePath, TargetImagePath)
+				.FileDescription(Info.IconDescription)
+				.RequiredSize(Info.IconRequiredSize)
+				.MaxDisplaySize(MaxDisplaySize)
+			]
+		];
+}
+
+void FIOSTargetSettingsCustomization::FindRequiredFiles()
+{
+	const UIOSRuntimeSettings& Settings = *GetDefault<UIOSRuntimeSettings>();
+	FString BundleIdentifier = Settings.BundleIdentifier.Replace(*gProjectNameText, FApp::GetGameName());
+	BundleIdentifier = BundleIdentifier.Replace(TEXT("_"), TEXT(""));
+#if PLATFORM_MAC
+	FString CmdExe = TEXT("/bin/sh");
+	FString ScriptPath = FPaths::ConvertRelativePathToFull(FPaths::EngineDir() / TEXT("Build/BatchFiles/Mac/RunMono.sh"));
+	FString IPPPath = FPaths::ConvertRelativePathToFull(FPaths::EngineDir() / TEXT("Binaries/DotNET/IOS/IPhonePackager.exe"));
+	FString CommandLine = FString::Printf(TEXT("\"%s\" \"%s\" certificates Engine -bundlename \"%s\""), *ScriptPath, *IPPPath, *(BundleIdentifier));
+#else
+	FString CmdExe = FPaths::ConvertRelativePathToFull(FPaths::EngineDir() / TEXT("Binaries/DotNET/IOS/IPhonePackager.exe"));
+	FString CommandLine = FString::Printf(TEXT("certificates Engine -bundlename \"%s\""), *(BundleIdentifier));
+#endif
+	IPPProcess = MakeShareable(new FMonitoredProcess(CmdExe, CommandLine, true));
+	OutputMessage = TEXT("");
+	IPPProcess->OnOutput().BindStatic(&OnOutput);
+	IPPProcess->Launch();
+	TickerHandle = FTicker::GetCoreTicker().AddTicker(FTickerDelegate::CreateRaw(this, &FIOSTargetSettingsCustomization::UpdateStatusDelegate), 1.0f);
+	if (ProvisionInfoSwitcher.IsValid())
+	{
+		ProvisionInfoSwitcher->SetActiveWidgetIndex(0);
+	}
+	if (CertificateInfoSwitcher.IsValid())
+	{
+		CertificateInfoSwitcher->SetActiveWidgetIndex(0);
+	}
+	RunningIPPProcess = true;
+}
+
+FReply FIOSTargetSettingsCustomization::OnInstallProvisionClicked()
+{
+	// pass the file to IPP to install
+	FString ProjectPath = FPaths::IsProjectFilePathSet() ? FPaths::ConvertRelativePathToFull(FPaths::GetProjectFilePath()) : FPaths::RootDir() / FApp::GetGameName() / FApp::GetGameName() + TEXT(".uproject");
+	FString ProvisionPath;
+
+	// get the provision by popping up the file dialog
+	TArray<FString> OpenFilenames;
+	IDesktopPlatform* DesktopPlatform = FDesktopPlatformModule::Get();
+	bool bOpened = false;
+	int32 FilterIndex = -1;
+	FString FileTypes = TEXT("Provision Files (*.mobileprovision)|*.mobileprovision");
+
+	if ( DesktopPlatform )
+	{
+		bOpened = DesktopPlatform->OpenFileDialog(
+			FSlateApplication::Get().FindBestParentWindowHandleForDialogs(nullptr),
+			LOCTEXT("ImportProvisionDialogTitle", "Import Provision").ToString(),
+			FPaths::GetProjectFilePath(),
+			TEXT(""),
+			FileTypes,
+			EFileDialogFlags::None,
+			OpenFilenames,
+			FilterIndex
+			);
+	}
+
+	if ( bOpened )
+	{
+		ProvisionPath = FPaths::ConvertRelativePathToFull(OpenFilenames[0]);
+
+		// see if the provision is already installed
+		FString DestName = FPaths::GetBaseFilename(ProvisionPath);
+		TCHAR Path[4096];
+#if PLATFORM_MAC
+		FPlatformMisc::GetEnvironmentVariable(TEXT("HOME"), Path, ARRAY_COUNT(Path));
+		FString Destination = FString::Printf(TEXT("\"%s/Library/MobileDevice/Provisioning Profiles/%s.mobileprovision\""), Path, *DestName);
+		FString Destination2 = FString::Printf(TEXT("\"%s/Library/MobileDevice/Provisioning Profiles/%s.mobileprovision\""), Path, FApp::GetGameName());
+#else
+		FPlatformMisc::GetEnvironmentVariable(TEXT("LOCALAPPDATA"), Path, ARRAY_COUNT(Path));
+		FString Destination = FString::Printf(TEXT("%s\\Apple Computer\\MobileDevice\\Provisioning Profiles\\%s.mobileprovision"), Path, *DestName);
+		FString Destination2 = FString::Printf(TEXT("%s\\Apple Computer\\MobileDevice\\Provisioning Profiles\\%s.mobileprovision"), Path, FApp::GetGameName());
+#endif
+		if (FPaths::FileExists(Destination) || FPaths::FileExists(Destination2))
+		{
+			FString MessagePrompt = FString::Printf(TEXT("%s mobile provision file already exists.  Do you want to replace this provision?"), *DestName);
+			if (FPlatformMisc::MessageBoxExt(EAppMsgType::OkCancel, *MessagePrompt, TEXT("File Exists")) == EAppReturnType::Cancel)
+			{
+				return FReply::Handled();
+			}
+		}
+
+		const UIOSRuntimeSettings& Settings = *GetDefault<UIOSRuntimeSettings>();
+		FString BundleIdentifier = Settings.BundleIdentifier.Replace(*gProjectNameText, FApp::GetGameName());
+		BundleIdentifier = BundleIdentifier.Replace(TEXT("_"), TEXT(""));
+#if PLATFORM_MAC
+		FString CmdExe = TEXT("/bin/sh");
+		FString ScriptPath = FPaths::ConvertRelativePathToFull(FPaths::EngineDir() / TEXT("Build/BatchFiles/Mac/RunMono.sh"));
+		FString IPPPath = FPaths::ConvertRelativePathToFull(FPaths::EngineDir() / TEXT("Binaries/DotNET/IOS/IPhonePackager.exe"));
+		FString CommandLine = FString::Printf(TEXT("\"%s\" \"%s\" Install Engine -project \"%s\" -provision \"%s\" -bundlename \"%s\""), *ScriptPath, *IPPPath, *ProjectPath, *ProvisionPath, *BundleIdentifier);
+#else
+		FString CmdExe = FPaths::ConvertRelativePathToFull(FPaths::EngineDir() / TEXT("Binaries/DotNET/IOS/IPhonePackager.exe"));
+		FString CommandLine = FString::Printf(TEXT("Install Engine -project \"%s\" -provision \"%s\" -bundlename \"%s\""), *ProjectPath, *ProvisionPath, *BundleIdentifier);
+#endif
+		IPPProcess = MakeShareable(new FMonitoredProcess(CmdExe, CommandLine, true));
+		OutputMessage = TEXT("");
+		IPPProcess->OnOutput().BindStatic(&OnOutput);
+		IPPProcess->Launch();
+		TickerHandle = FTicker::GetCoreTicker().AddTicker(FTickerDelegate::CreateRaw(this, &FIOSTargetSettingsCustomization::UpdateStatusDelegate), 10.0f);
+		if (ProvisionInfoSwitcher.IsValid())
+		{
+			ProvisionInfoSwitcher->SetActiveWidgetIndex(1);
+		}
+		RunningIPPProcess = true;
+	}
+
+	return FReply::Handled();
+}
+
+FReply FIOSTargetSettingsCustomization::OnInstallCertificateClicked()
+{
+	// pass the file to IPP to install
+	FString ProjectPath = FPaths::IsProjectFilePathSet() ? FPaths::ConvertRelativePathToFull(FPaths::GetProjectFilePath()) : FPaths::RootDir() / FApp::GetGameName() / FApp::GetGameName() + TEXT(".uproject");
+	FString CertPath;
+
+	// get the provision by popping up the file dialog
+	TArray<FString> OpenFilenames;
+	IDesktopPlatform* DesktopPlatform = FDesktopPlatformModule::Get();
+	bool bOpened = false;
+	int32 FilterIndex = -1;
+	FString FileTypes = TEXT("Code Signing Certificates (*.cer;*.p12)|*.cer;*p12");
+
+	if ( DesktopPlatform )
+	{
+		bOpened = DesktopPlatform->OpenFileDialog(
+			FSlateApplication::Get().FindBestParentWindowHandleForDialogs(nullptr),
+			LOCTEXT("ImportCertificateDialogTitle", "Import Certificate").ToString(),
+			FPaths::GetProjectFilePath(),
+			TEXT(""),
+			FileTypes,
+			EFileDialogFlags::None,
+			OpenFilenames,
+			FilterIndex
+			);
+	}
+
+	if ( bOpened )
+	{
+		const UIOSRuntimeSettings& Settings = *GetDefault<UIOSRuntimeSettings>();
+		FString BundleIdentifier = Settings.BundleIdentifier.Replace(*gProjectNameText, FApp::GetGameName());
+		BundleIdentifier = BundleIdentifier.Replace(TEXT("_"), TEXT(""));
+		CertPath = FPaths::ConvertRelativePathToFull(OpenFilenames[0]);
+#if PLATFORM_MAC
+		FString CmdExe = TEXT("/bin/sh");
+		FString ScriptPath = FPaths::ConvertRelativePathToFull(FPaths::EngineDir() / TEXT("Build/BatchFiles/Mac/RunMono.sh"));
+		FString IPPPath = FPaths::ConvertRelativePathToFull(FPaths::EngineDir() / TEXT("Binaries/DotNET/IOS/IPhonePackager.exe"));
+		FString CommandLine = FString::Printf(TEXT("\"%s\" \"%s\" Install Engine -project \"%s\" -certificate \"%s\" -bundlename \"%s\""), *ScriptPath, *IPPPath, *ProjectPath, *CertPath, *BundleIdentifier);
+#else
+		FString CmdExe = FPaths::ConvertRelativePathToFull(FPaths::EngineDir() / TEXT("Binaries/DotNET/IOS/IPhonePackager.exe"));
+		FString CommandLine = FString::Printf(TEXT("Install Engine -project \"%s\" -certificate \"%s\" -bundlename \"%s\""), *ProjectPath, *CertPath, *BundleIdentifier);
+#endif
+		IPPProcess = MakeShareable(new FMonitoredProcess(CmdExe, CommandLine, true));
+		OutputMessage = TEXT("");
+		IPPProcess->OnOutput().BindStatic(&OnOutput);
+		IPPProcess->Launch();
+		TickerHandle = FTicker::GetCoreTicker().AddTicker(FTickerDelegate::CreateRaw(this, &FIOSTargetSettingsCustomization::UpdateStatusDelegate), 10.0f);
+		if (CertificateInfoSwitcher.IsValid())
+		{
+			CertificateInfoSwitcher->SetActiveWidgetIndex(1);
+		}
+		RunningIPPProcess = true;
+	}
+
+	return FReply::Handled();
+}
+
+FReply FIOSTargetSettingsCustomization::OnCertificateRequestClicked()
+{
+	// TODO: bring up an open file dialog and then install the provision
+	return FReply::Handled();
+}
+
+FReply FIOSTargetSettingsCustomization::OnGenerateSSHKey()
+{
+	// see if the key is already generated
+	const UIOSRuntimeSettings& Settings = *GetDefault<UIOSRuntimeSettings>();
+
+	FString RemoteServerAddress;
+	FString RemoteServerPort;
+	int32	colonIndex;
+
+	if(Settings.RemoteServerName.FindChar(':', colonIndex))
+	{
+		RemoteServerAddress = Settings.RemoteServerName.Left(colonIndex);
+		RemoteServerPort = Settings.RemoteServerName.RightChop(colonIndex + 1);
+	}
+	else
+	{
+		RemoteServerAddress = Settings.RemoteServerName;
+		RemoteServerPort = "22";
+	}
+
+	TCHAR Path[4096];
+	FPlatformMisc::GetEnvironmentVariable(TEXT("APPDATA"), Path, ARRAY_COUNT(Path));
+	FString Destination = FString::Printf(TEXT("%s\\Unreal Engine\\UnrealBuildTool\\SSHKeys\\%s\\%s\\RemoteToolChainPrivate.key"), Path, *RemoteServerAddress, *(Settings.RSyncUsername));
+	if (FPaths::FileExists(Destination))
+	{
+		FString MessagePrompt = FString::Printf(TEXT("An SSH Key already exists.  Do you want to replace this key?"));
+		if (FPlatformMisc::MessageBoxExt(EAppMsgType::OkCancel, *MessagePrompt, TEXT("Key Exists")) == EAppReturnType::Cancel)
+		{
+			return FReply::Handled();
+		}
+	}
+
+	FString CmdExe = FPaths::ConvertRelativePathToFull(FPaths::EngineDir() / TEXT("Build/BatchFiles/MakeAndInstallSSHKey.bat"));
+	FString DeltaCopyPath = Settings.DeltaCopyInstallPath.Path;
+	if (DeltaCopyPath.IsEmpty() || !FPaths::DirectoryExists(DeltaCopyPath))
+	{
+		// If no user specified directory try the UE4 bundled directory
+		DeltaCopyPath = FPaths::ConvertRelativePathToFull(FPaths::EngineDir() / TEXT("Extras\\ThirdPartyNotUE\\DeltaCopy\\Binaries"));
+	}
+
+	if (!FPaths::DirectoryExists(DeltaCopyPath))
+	{
+		// if no UE4 bundled version of DeltaCopy, try and use the default install location
+		TCHAR ProgramPath[4096];
+		FPlatformMisc::GetEnvironmentVariable(TEXT("PROGRAMFILES(X86)"), ProgramPath, ARRAY_COUNT(ProgramPath));
+		DeltaCopyPath = FPaths::Combine(ProgramPath, TEXT("DeltaCopy"));
+	}
+	
+	if (!FPaths::DirectoryExists(DeltaCopyPath))
+	{
+		UE_LOG(LogIOSTargetSettings, Error, TEXT("DeltaCopy is not installed correctly"));
+	}
+
+	FString CygwinPath = TEXT("/cygdrive/") + FString(Path).Replace(TEXT(":"), TEXT("")).Replace(TEXT("\\"), TEXT("/"));
+	FString EnginePath = FPaths::EngineDir();
+	FString CommandLine = FString::Printf(TEXT("\"%s/ssh.exe\" %s \"%s\\rsync.exe\" %s %s \"%s\" \"%s\" \"%s\""),
+		*DeltaCopyPath,
+		*RemoteServerPort,
+		*DeltaCopyPath,
+		*(Settings.RSyncUsername),
+		*RemoteServerAddress,
+		Path,
+		*CygwinPath,
+		*EnginePath);
+
+	OutputMessage = TEXT("");
+	IPPProcess = MakeShareable(new FMonitoredProcess(CmdExe, CommandLine, false, false));
+	IPPProcess->Launch();
+	TickerHandle = FTicker::GetCoreTicker().AddTicker(FTickerDelegate::CreateRaw(this, &FIOSTargetSettingsCustomization::UpdateStatusDelegate), 10.0f);
+	RunningIPPProcess = true;
+
+	return FReply::Handled();
+}
+
+const FSlateBrush* FIOSTargetSettingsCustomization::GetProvisionStatus() const
+{
+	if( bProvisionInstalled )
+	{
+		return FEditorStyle::GetBrush("Automation.Success");
+	}
+	else
+	{
+		return FEditorStyle::GetBrush("Automation.Fail");
+	}
+}
+
+const FSlateBrush* FIOSTargetSettingsCustomization::GetCertificateStatus() const
+{
+	if( bCertificateInstalled )
+	{
+		return FEditorStyle::GetBrush("Automation.Success");
+	}
+	else
+	{
+		return FEditorStyle::GetBrush("Automation.Fail");
+	}
+}
+
+bool FIOSTargetSettingsCustomization::UpdateStatusDelegate(float DeltaTime)
+{
+	if (IPPProcess.IsValid())
+	{
+		if (IPPProcess->Update())
+		{
+			return true;
+		}
+		int RetCode = IPPProcess->GetReturnCode();
+		IPPProcess = NULL;
+		UpdateStatus();
+		UpdateSSHStatus();
+	}
+	RunningIPPProcess = false;
+
+	return false;
+}
+
+TSharedRef<ITableRow> FIOSTargetSettingsCustomization::HandleProvisionListGenerateRow( ProvisionPtr InProvision, const TSharedRef<STableViewBase>& OwnerTable )
+{
+	return SNew(SProvisionListRow, OwnerTable)
+		.Provision(InProvision)
+		.ProvisionList(ProvisionList)
+		.OnProvisionChanged(this, &FIOSTargetSettingsCustomization::HandleProvisionChanged);
+}
+
+void FIOSTargetSettingsCustomization::HandleProvisionChanged(FString Provision)
+{
+	FText OutText;
+	MobileProvisionProperty->GetValueAsFormattedText(OutText);
+	if (OutText.ToString() != Provision)
+	{
+		MobileProvisionProperty->SetValueFromFormattedString(Provision);
+	}
+	SignCertificateProperty->GetValueAsFormattedText(OutText);
+	if (Provision == TEXT("") && OutText.ToString() == TEXT(""))
+	{
+		bManuallySelected = false;
+		FilterLists();
+	}
+	else if (!bManuallySelected)
+	{
+		bManuallySelected = true;
+		FilterLists();
+	}
+}
+
+void FIOSTargetSettingsCustomization::HandleCertificateChanged(FString Certificate)
+{
+	FText OutText;
+	SignCertificateProperty->GetValueAsFormattedText(OutText);
+	if (OutText.ToString() != Certificate)
+	{
+		SignCertificateProperty->SetValueFromFormattedString(Certificate);
+	}
+	MobileProvisionProperty->GetValueAsFormattedText(OutText);
+	if (Certificate == TEXT("") && OutText.ToString() == TEXT(""))
+	{
+		bManuallySelected = false;
+		FilterLists();
+	}
+	else if (!bManuallySelected)
+	{
+		bManuallySelected = true;
+		FilterLists();
+	}
+}
+
+TSharedRef<ITableRow> FIOSTargetSettingsCustomization::HandleCertificateListGenerateRow( CertificatePtr InCertificate, const TSharedRef<STableViewBase>& OwnerTable )
+{
+	return SNew(SCertificateListRow, OwnerTable)
+		.Certificate(InCertificate)
+		.CertificateList(CertificateList)
+		.OnCertificateChanged(this, &FIOSTargetSettingsCustomization::HandleCertificateChanged);
+}
+
+void FIOSTargetSettingsCustomization::HandleAllProvisionsHyperlinkNavigate( bool AllProvisions )
+{
+	bShowAllProvisions = AllProvisions;
+	FilterLists();
+}
+
+void FIOSTargetSettingsCustomization::HandleAllCertificatesHyperlinkNavigate( bool AllCertificates )
+{
+	bShowAllCertificates = AllCertificates;
+	FilterLists();
+}
+
+void FIOSTargetSettingsCustomization::FilterLists()
+{
+	FilteredProvisionList.Reset();
+	FilteredCertificateList.Reset();
+
+	for (int Index = 0; Index < ProvisionList->Num(); ++Index)
+	{
+		if (SelectedProvision.Contains((*ProvisionList)[Index]->Name) && SelectedFile.Contains((*ProvisionList)[Index]->FileName) && !bManuallySelected)
+		{
+			(*ProvisionList)[Index]->bSelected = true;
+		}
+		else
+		{
+			(*ProvisionList)[Index]->bSelected = false;
+		}
+		if (bShowAllProvisions || (*ProvisionList)[Index]->Status.Contains("VALID"))
+		{
+			FilteredProvisionList.Add((*ProvisionList)[Index]);
+		}
+	}
+
+	if (ProvisionList->Num() > 0)
+	{
+		if (ProvisionInfoSwitcher.IsValid())
+		{
+			ProvisionInfoSwitcher->SetActiveWidgetIndex(3);
+		}
+		if (FilteredProvisionList.Num() == 0 && !bShowAllProvisions)
+		{
+			FilteredProvisionList.Append(*ProvisionList);
+		}
+	}
+	else
+	{
+		if (ProvisionInfoSwitcher.IsValid())
+		{
+			ProvisionInfoSwitcher->SetActiveWidgetIndex(2);
+		}
+	}
+
+	for (int Index = 0; Index < CertificateList->Num(); ++Index)
+	{
+		if (SelectedCert.Contains((*CertificateList)[Index]->Name) && !bManuallySelected)
+		{
+			(*CertificateList)[Index]->bSelected = true;
+		}
+		else
+		{
+			(*CertificateList)[Index]->bSelected = false;
+		}
+		if (bShowAllCertificates || (*CertificateList)[Index]->Status.Contains("VALID"))
+		{
+			FilteredCertificateList.Add((*CertificateList)[Index]);
+		}
+	}
+
+	if (CertificateList->Num() > 0)
+	{
+		if (CertificateInfoSwitcher.IsValid())
+		{
+			CertificateInfoSwitcher->SetActiveWidgetIndex(3);
+		}
+		if (FilteredCertificateList.Num() == 0 && !bShowAllCertificates)
+		{
+			FilteredCertificateList.Append(*CertificateList);
+		}
+	}
+	else
+	{
+		if (CertificateInfoSwitcher.IsValid())
+		{
+			CertificateInfoSwitcher->SetActiveWidgetIndex(2);
+		}
+	}
+
+	CertificateListView->RequestListRefresh();
+	ProvisionListView->RequestListRefresh();
+}
+
+bool FIOSTargetSettingsCustomization::IsImportEnabled() const
+{
+	return !RunningIPPProcess.Get();
+}
+
+void FIOSTargetSettingsCustomization::OnBundleIdentifierChanged(const FText& NewText, ETextCommit::Type CommitType, TSharedRef<IPropertyHandle> InPropertyHandle)
+{
+	if(!IsBundleIdentifierValid(NewText.ToString()))
+	{
+		BundleIdTextBox->SetError( LOCTEXT("NameContainsInvalidCharacters", "Identifier may only contain the characters 0-9, A-Z, a-z, period, hyphen, or [PROJECT_NAME]") );
+	}
+	else
+	{
+		BundleIdTextBox->SetError(FText::GetEmpty());
+
+		FText OutText;
+		InPropertyHandle->GetValueAsFormattedText(OutText);
+		if (OutText.ToString() != NewText.ToString())
+		{
+			InPropertyHandle->SetValueFromFormattedString( NewText.ToString() );
+			FindRequiredFiles();
+		}
+	}
+}
+
+void FIOSTargetSettingsCustomization::OnBundleIdentifierTextChanged(const FText& NewText, ETextCommit::Type CommitType, TSharedRef<IPropertyHandle> InPropertyHandle)
+{
+	if(!IsBundleIdentifierValid(NewText.ToString()))
+	{
+		BundleIdTextBox->SetError( LOCTEXT("NameContainsInvalidCharacters", "Identifier may only contain the characters 0-9, A-Z, a-z, period, hyphen, or [PROJECT_NAME]") );
+	}
+	else
+	{
+		BundleIdTextBox->SetError(FText::GetEmpty());
+	}
+}
+
+bool FIOSTargetSettingsCustomization::IsBundleIdentifierValid(const FString& inIdentifier)
+{
+	for(int32 i = 0; i < inIdentifier.Len(); ++i)
+	{
+		TCHAR	c = inIdentifier[i];
+		
+		if(c == '[')
+		{
+			if(inIdentifier.Find(gProjectNameText, ESearchCase::CaseSensitive, ESearchDir::FromStart, i) != i)
+			{
+				return false;
+			}
+			i += gProjectNameText.Len();
+		}
+		else if((c < '0' || c > '9') && (c < 'a' || c > 'z') && (c < 'A' || c > 'Z') && c != '.' && c != '-')
+		{
+			return false;
+		}
+	}
+
+	return true;
+}
+
+void FIOSTargetSettingsCustomization::OnRemoteServerChanged(const FText& NewText, ETextCommit::Type CommitType, TSharedRef<IPropertyHandle> InPropertyHandle)
+{
+	FText OutText;
+	InPropertyHandle->GetValueAsFormattedText(OutText);
+	if (OutText.ToString() != NewText.ToString())
+	{
+		InPropertyHandle->SetValueFromFormattedString(NewText.ToString());
+		OutputMessage = TEXT("");
+		UpdateSSHStatus();
+	}
+}
+
+FText FIOSTargetSettingsCustomization::GetBundleText(TSharedRef<IPropertyHandle> InPropertyHandle) const
+{
+	FText OutText;
+	InPropertyHandle->GetValueAsFormattedText(OutText);
+	return OutText;
+}
+
+TSharedRef<SWidget> FIOSTargetSettingsCustomization::OnGetShaderVersionContent()
+{
+	FMenuBuilder MenuBuilder(true, NULL);
+	
+	UEnum* Enum = FindObjectChecked<UEnum>(ANY_PACKAGE, TEXT("EIOSMetalShaderStandard"), true);
+	
+	for (int32 i = 0; i < Enum->GetMaxEnumValue(); i++)
+	{
+		if (Enum->IsValidEnumValue(i))
+		{
+			FUIAction ItemAction(FExecuteAction::CreateSP(this, &FIOSTargetSettingsCustomization::SetShaderStandard, i));
+			MenuBuilder.AddMenuEntry(Enum->GetDisplayNameTextByValue(i), TAttribute<FText>(), FSlateIcon(), ItemAction);
+		}
+	}
+	
+	return MenuBuilder.MakeWidget();
+}
+
+FText FIOSTargetSettingsCustomization::GetShaderVersionDesc() const
+{
+	uint8 EnumValue;
+	ShaderVersionPropertyHandle->GetValue(EnumValue);
+	
+	UEnum* Enum = FindObjectChecked<UEnum>(ANY_PACKAGE, TEXT("EIOSMetalShaderStandard"), true);
+	
+	if (EnumValue < Enum->GetMaxEnumValue() && Enum->IsValidEnumValue(EnumValue))
+	{
+		return Enum->GetDisplayNameTextByValue(EnumValue);
+	}
+	
+	return FText::GetEmpty();
+}
+
+TSharedRef<SWidget> FIOSTargetSettingsCustomization::OnGetMinVersionContent()
+{
+	FMenuBuilder MenuBuilder(true, NULL);
+
+	UEnum* Enum = FindObjectChecked<UEnum>(ANY_PACKAGE, TEXT("EIOSVersion"), true);
+
+	for (int32 i = 0; i < Enum->GetMaxEnumValue(); i++)
+	{
+		if (Enum->IsValidEnumValue(i) && !Enum->HasMetaData(TEXT("Hidden"), Enum->GetIndexByValue(i)))
+		{
+			FUIAction ItemAction(FExecuteAction::CreateSP(this, &FIOSTargetSettingsCustomization::SetMinVersion, i));
+			MenuBuilder.AddMenuEntry(Enum->GetDisplayNameTextByValue(i), TAttribute<FText>(), FSlateIcon(), ItemAction);
+		}
+	}
+
+	return MenuBuilder.MakeWidget();
+}
+
+FText FIOSTargetSettingsCustomization::GetMinVersionDesc() const
+{
+	uint8 EnumValue;
+	MinOSPropertyHandle->GetValue(EnumValue);
+
+	UEnum* Enum = FindObjectChecked<UEnum>(ANY_PACKAGE, TEXT("EIOSVersion"), true);
+
+	if (EnumValue < Enum->GetMaxEnumValue() && Enum->IsValidEnumValue(EnumValue))
+	{
+		return Enum->GetDisplayNameTextByValue(EnumValue);
+	}
+
+	return FText::GetEmpty();
+}
+
+void FIOSTargetSettingsCustomization::SetShaderStandard(int32 Value)
+{
+    if (Value >= 1 && MinOSPropertyHandle.IsValid() && GLES2PropertyHandle.IsValid())
+    {
+        FText Message;
+        
+		uint8 EnumValue = (uint8)EIOSVersion::IOS_10;
+		if (MinOSPropertyHandle.IsValid())
+		{
+		MinOSPropertyHandle->GetValue(EnumValue);
+		}
+		
+		bool bHasGL = false;
+		GLES2PropertyHandle->GetValue(bHasGL);
+		
+		if (Value == 1)
+        {
+			if ((EIOSVersion)EnumValue < EIOSVersion::IOS_9 && !bHasGL)
+			{
+				Message = LOCTEXT("iOSMetalShaderVersion1_1","Enabling Metal Shader Standard v1.1 increases the minimum operating system requirement for Metal from iOS 8.0 or later to iOS 9.0 or later - earlier OS versions would require OpenGL ES2. This does not affect tvOS.");
+			}
+        }
+        else if ((EIOSVersion)EnumValue < EIOSVersion::IOS_10 && !bHasGL)
+        {
+            Message = LOCTEXT("iOSMetalShaderVersion1_2","Enabling Metal Shader Standard v1.2 increases the minimum operating system requirement for Metal from iOS 8.0/tvOS 9.0 or later to iOS/tvOS 10.0 or later - earlier OS versions would require OpenGL ES2.");
+        }
+		
+		
+		ShaderVersionWarningTextBox->SetError(Message);
+	}
+	else
+	{
+		ShaderVersionWarningTextBox->SetError(TEXT(""));
+	}
+	FPropertyAccess::Result Res = ShaderVersionPropertyHandle->SetValue((uint8)Value);
+	check(Res == FPropertyAccess::Success);
+}
+
+void FIOSTargetSettingsCustomization::UpdateShaderStandardWarning()
+{
+	// Update the UI
+	uint8 EnumValue;
+	ShaderVersionPropertyHandle->GetValue(EnumValue);
+	SetShaderStandard(EnumValue);
+}
+
+void FIOSTargetSettingsCustomization::UpdateOSVersionWarning()
+{
+	bool ArchValue;
+	bool bEnabled = false;
+	DevArmV7PropertyHandle->GetValue(ArchValue);
+	bEnabled |= ArchValue;
+
+	DevArmV7sPropertyHandle->GetValue(ArchValue);
+	bEnabled |= ArchValue;
+
+	ShipArmV7PropertyHandle->GetValue(ArchValue);
+	bEnabled |= ArchValue;
+
+	ShipArmV7sPropertyHandle->GetValue(ArchValue);
+	bEnabled |= ArchValue;
+
+	FText Message;
+	Message = LOCTEXT("IOSDeprecation", "Enabling ArmV7 or ArmV7S will no longer be supported in 4.17.");
+
+	// Update the UI
+	if (bEnabled)
+	{
+		IOSVersionWarningTextBox->SetError(Message);
+	}
+	else
+	{
+		IOSVersionWarningTextBox->SetError(TEXT(""));
+	}
+}
+
+void FIOSTargetSettingsCustomization::UpdateGLVersionWarning()
+{
+	bool bEnabled = false;
+	GLES2PropertyHandle->GetValue(bEnabled);
+
+	FText Message;
+	Message = LOCTEXT("GLES2Deprecation", "GLES2 will no longer be supported in 4.17.");
+
+	// Update the UI
+	if (bEnabled)
+	{
+		GLVersionWarningTextBox->SetError(Message);
+	}
+	else
+	{
+		GLVersionWarningTextBox->SetError(TEXT(""));
+	}
+
+	UpdateShaderStandardWarning();
+}
+
+void FIOSTargetSettingsCustomization::SetMinVersion(int32 Value)
+{
+	FPropertyAccess::Result Res = MinOSPropertyHandle->SetValue((uint8)Value);
+	check(Res == FPropertyAccess::Success);
+}
+
+void FIOSTargetSettingsCustomization::HandleGLES2CheckBoxCheckStateChanged(ECheckBoxState NewState)
+{
+	GLES2PropertyHandle->SetValue(NewState == ECheckBoxState::Checked ? true : false);
+
+	UpdateGLVersionWarning();
+}
+
+//////////////////////////////////////////////////////////////////////////
+
+#undef LOCTEXT_NAMESPACE
diff --git a/Engine/Source/Developer/IOS/IOSPlatformEditor/Private/IOSTargetSettingsCustomization.h b/Engine/Source/Developer/IOS/IOSPlatformEditor/Private/IOSTargetSettingsCustomization.h
new file mode 100644
--- /dev/null
+++ b/Engine/Source/Developer/IOS/IOSPlatformEditor/Private/IOSTargetSettingsCustomization.h
@@ -0,0 +1,228 @@
+// Copyright 1998-2017 Epic Games, Inc. All Rights Reserved.
+
+#pragma once
+
+#include "CoreMinimal.h"
+#include "SlateFwd.h"
+#include "Misc/Attribute.h"
+#include "Input/Reply.h"
+#include "Widgets/Views/STableViewBase.h"
+#include "Widgets/Views/STableRow.h"
+#include "IDetailCustomization.h"
+#include "ShaderFormatsPropertyDetails.h"
+
+class FMonitoredProcess;
+class IDetailLayoutBuilder;
+class IPropertyHandle;
+class SEditableTextBox;
+class SErrorText;
+
+//////////////////////////////////////////////////////////////////////////
+// FProvision structure
+class FProvision
+{
+public:
+	FString Name;
+	FString FileName;
+	FString Status;
+	bool bDistribution;
+	bool bSelected;
+	bool bManuallySelected;
+};
+typedef TSharedPtr<class FProvision> ProvisionPtr;
+typedef TSharedPtr<TArray<ProvisionPtr>> ProvisionListPtr;
+
+//////////////////////////////////////////////////////////////////////////
+// FCertificate structure
+class FCertificate
+{
+public:
+	FString Name;
+	FString Status;
+	FString Expires;
+	bool bSelected;
+	bool bManuallySelected;
+};
+typedef TSharedPtr<class FCertificate> CertificatePtr;
+typedef TSharedPtr<TArray<CertificatePtr>> CertificateListPtr;
+
+//////////////////////////////////////////////////////////////////////////
+// FIOSTargetSettingsCustomization
+
+class FIOSTargetSettingsCustomization : public IDetailCustomization
+{
+public:
+	~FIOSTargetSettingsCustomization();
+
+	// Makes a new instance of this detail layout class for a specific detail view requesting it
+	static TSharedRef<IDetailCustomization> MakeInstance();
+
+	// IDetailCustomization interface
+	virtual void CustomizeDetails(IDetailLayoutBuilder& DetailLayout) override;
+	// End of IDetailCustomization interface
+
+private:
+	TArray<struct FPlatformIconInfo> IconNames;
+	TArray<struct FPlatformIconInfo> LaunchImageNames;
+
+	const FString EngineInfoPath;
+	const FString GameInfoPath;
+
+	const FString EngineGraphicsPath;
+	const FString GameGraphicsPath;
+
+	IDetailLayoutBuilder* SavedLayoutBuilder;
+
+	// Is the manifest writable?
+	TAttribute<bool> SetupForPlatformAttribute;
+
+	bool bProvisionInstalled;
+	bool bCertificateInstalled;
+	bool bShowAllProvisions;
+	bool bShowAllCertificates;
+	bool bManuallySelected;
+
+	TSharedPtr<FMonitoredProcess> IPPProcess;
+	FDelegateHandle TickerHandle;
+	TSharedPtr<TArray<ProvisionPtr>> ProvisionList;
+	TArray<ProvisionPtr> FilteredProvisionList;
+	TSharedPtr<SListView<ProvisionPtr> > ProvisionListView;
+	TSharedPtr<SWidgetSwitcher > ProvisionInfoSwitcher;
+	TSharedPtr<TArray<CertificatePtr>> CertificateList;
+	TArray<CertificatePtr> FilteredCertificateList;
+	TSharedPtr<SListView<CertificatePtr> > CertificateListView;
+	TSharedPtr<SWidgetSwitcher > CertificateInfoSwitcher;
+	TAttribute<bool> RunningIPPProcess;
+	TSharedPtr<IPropertyHandle> MobileProvisionProperty;
+	TSharedPtr<IPropertyHandle> SignCertificateProperty;
+	TSharedPtr<IPropertyHandle> ShaderVersionPropertyHandle;
+	TSharedPtr<IPropertyHandle> MinOSPropertyHandle;
+	TSharedPtr<IPropertyHandle> GLES2PropertyHandle;
+	TSharedPtr<IPropertyHandle> DevArmV7PropertyHandle;
+	TSharedPtr<IPropertyHandle> DevArmV7sPropertyHandle;
+	TSharedPtr<IPropertyHandle> ShipArmV7PropertyHandle;
+	TSharedPtr<IPropertyHandle> ShipArmV7sPropertyHandle;
+	TSharedPtr<IPropertyHandle> AutomaticSigningProperty;
+
+	FString SelectedProvision;
+	FString SelectedFile;
+	FString SelectedCert;
+
+private:
+	FIOSTargetSettingsCustomization();
+
+	void BuildPListSection(IDetailLayoutBuilder& DetailLayout);
+	void BuildIconSection(IDetailLayoutBuilder& DetailLayout);
+    void BuildRemoteBuildingSection(IDetailLayoutBuilder& DetailLayout);
+
+	// Navigates to the plist in explorer or finder
+	FReply OpenPlistFolder();
+
+	// Copies the setup files for the platform into the project
+	void CopySetupFilesIntoProject();
+
+	// Builds an image row
+	void BuildImageRow(class IDetailLayoutBuilder& DetailLayout, class IDetailCategoryBuilder& Category, const struct FPlatformIconInfo& Info, const FVector2D& MaxDisplaySize);
+
+	// Install the provision
+ 	FReply OnInstallProvisionClicked();
+
+	// Install the provision
+	FReply OnInstallCertificateClicked();
+
+	// certificate request
+	FReply OnCertificateRequestClicked();
+
+	// ssh key request
+	FReply OnGenerateSSHKey();
+
+	// Get the image to display for the provision status
+	const FSlateBrush* GetProvisionStatus() const;
+
+	// Get the image to display for the certificate status
+	const FSlateBrush* GetCertificateStatus() const;
+
+	// find the installed certificates and provisions
+	void FindRequiredFiles();
+
+	// Update the provision status
+	void UpdateStatus();
+
+	// Update the ssh key status
+	void UpdateSSHStatus();
+
+	// status tick delay
+	bool UpdateStatusDelegate(float delay);
+
+	// handle provision list generate row
+	TSharedRef<ITableRow> HandleProvisionListGenerateRow( ProvisionPtr Provision, const TSharedRef<STableViewBase>& OwnerTable );
+
+	// handle certificate list generate row
+	TSharedRef<ITableRow> HandleCertificateListGenerateRow( CertificatePtr Certificate, const TSharedRef<STableViewBase>& OwnerTable );
+
+	// handle which set of provisions to view
+	void HandleAllProvisionsHyperlinkNavigate( bool AllProvisions );
+
+	// handle which set of provisions to view
+	void HandleAllCertificatesHyperlinkNavigate( bool AllCertificates );
+
+	// filter the lists based on the settings
+	void FilterLists();
+
+	// returns whether we are importing or not
+	bool IsImportEnabled() const;
+
+	// updates the bundle identifier if it is valid and checks for a matching provision/certificate
+	void OnBundleIdentifierChanged(const FText& NewText, ETextCommit::Type, TSharedRef<IPropertyHandle> InPropertyHandle);
+
+	// posts an error if the bundle identifier has become invalid
+	void OnBundleIdentifierTextChanged(const FText& NewText, ETextCommit::Type, TSharedRef<IPropertyHandle> InPropertyHandle);
+
+	// returns true if the given string is a valid bundle identifier
+	bool IsBundleIdentifierValid(const FString& inIdentifier);
+
+	// updates the text in the ini file and checks for a valid provision/certificate
+	void OnRemoteServerChanged(const FText& NewText, ETextCommit::Type, TSharedRef<IPropertyHandle> InPropertyHandle);
+
+	void HandleProvisionChanged(FString Provision);
+
+	void HandleCertificateChanged(FString Certificate);
+	
+	/** Delegate handler to get the list of shader standards */
+	TSharedRef<SWidget> OnGetShaderVersionContent();
+	
+	/** Delegate handler to get the description of the shader standard */
+	FText GetShaderVersionDesc() const;
+
+	/** Delegate handler to get the list of shader standards */
+	TSharedRef<SWidget> OnGetMinVersionContent();
+
+	/** Delegate handler to get the description of the shader standard */
+	FText GetMinVersionDesc() const;
+
+	void SetShaderStandard(int32 Value);
+	
+	void UpdateShaderStandardWarning();
+	
+	void UpdateOSVersionWarning();
+
+	void UpdateGLVersionWarning();
+
+	void SetMinVersion(int32 Value);
+
+	void HandleGLES2CheckBoxCheckStateChanged(ECheckBoxState NewState);
+
+	// 
+	FText GetBundleText(TSharedRef<IPropertyHandle> InPropertyHandle) const;
+
+	TSharedPtr< SEditableTextBox > BundleIdTextBox;
+	
+	/** Reference to the shader version property warning text box. */
+	TSharedPtr< SErrorText > ShaderVersionWarningTextBox;
+
+	/** Reference to the shader version property warning text box. */
+	TSharedPtr< SErrorText > IOSVersionWarningTextBox;
+
+	/** Reference to the os version property warning text box. */
+	TSharedPtr< SErrorText > GLVersionWarningTextBox;
+};
diff --git a/Engine/Source/Developer/IOS/IOSPlatformEditor/Private/SCertificateListRow.h b/Engine/Source/Developer/IOS/IOSPlatformEditor/Private/SCertificateListRow.h
new file mode 100644
--- /dev/null
+++ b/Engine/Source/Developer/IOS/IOSPlatformEditor/Private/SCertificateListRow.h
@@ -0,0 +1,188 @@
+// Copyright 1998-2017 Epic Games, Inc. All Rights Reserved.
+
+#pragma once
+
+#include "CoreMinimal.h"
+#include "Styling/SlateColor.h"
+#include "Widgets/SWidget.h"
+#include "Layout/Margin.h"
+#include "Widgets/SNullWidget.h"
+#include "Widgets/Views/STableViewBase.h"
+#include "Widgets/DeclarativeSyntaxSupport.h"
+#include "Styling/SlateTypes.h"
+#include "Widgets/Layout/SBox.h"
+#include "Widgets/Views/STableRow.h"
+#include "IOSTargetSettingsCustomization.h"
+#include "Widgets/Text/STextBlock.h"
+#include "Widgets/Input/SCheckBox.h"
+#include "Widgets/Views/SListView.h"
+#include "SlateOptMacros.h"
+
+DECLARE_DELEGATE_OneParam(FOnCertificateChanged, FString);
+
+/**
+ * Implements a row widget for the certificate list view.
+ */
+class SCertificateListRow
+	: public SMultiColumnTableRow<CertificatePtr>
+{
+public:
+
+	SLATE_BEGIN_ARGS(SCertificateListRow) { }
+		SLATE_ARGUMENT(CertificatePtr, Certificate)
+		SLATE_ARGUMENT(CertificateListPtr, CertificateList)
+		SLATE_EVENT(FOnCertificateChanged, OnCertificateChanged)
+	SLATE_END_ARGS()
+
+public:
+
+	/**
+	 * Constructs the widget.
+	 *
+	 * @param InArgs The arguments.
+	 */
+	void Construct( const FArguments& InArgs, const TSharedRef<STableViewBase>& InOwnerTableView )
+	{
+		Certificate = InArgs._Certificate;
+		CertificateList = InArgs._CertificateList;
+		OnCertificateChanged_Handler = InArgs._OnCertificateChanged;
+		
+		SMultiColumnTableRow<CertificatePtr>::Construct(FSuperRowType::FArguments(), InOwnerTableView);
+	}
+
+public:
+
+	/**
+	 * Generates the widget for the specified column.
+	 *
+	 * @param ColumnName The name of the column to generate the widget for.
+	 * @return The widget.
+	 */
+	BEGIN_SLATE_FUNCTION_BUILD_OPTIMIZATION
+	virtual TSharedRef<SWidget> GenerateWidgetForColumn( const FName& ColumnName ) override
+	{
+		if (ColumnName == TEXT("Selected"))
+		{
+			return SNew(SCheckBox)
+				.IsChecked(this, &SCertificateListRow::HandleChecked)
+				.OnCheckStateChanged(this, &SCertificateListRow::HandleCheckStateChanged);
+		}
+		else if (ColumnName == TEXT("Name"))
+		{
+			return SNew(SBox)
+				.Padding(FMargin(4.0f, 0.0f))
+				.VAlign(VAlign_Center)
+				[
+					SNew(STextBlock)
+						.ColorAndOpacity(this, &SCertificateListRow::HandleSelectedColorAndOpacity)
+						.Text(this, &SCertificateListRow::HandleNameText)
+				];
+		}
+		else if (ColumnName == TEXT("Status"))
+		{
+			return SNew(SBox)
+				.Padding(FMargin(4.0f, 0.0f))
+				.VAlign(VAlign_Center)
+				[
+					SNew(STextBlock)
+						.ColorAndOpacity(this, &SCertificateListRow::HandleStatusTextColorAndOpacity)
+						.Text(this, &SCertificateListRow::HandleStatusTextBlockText)
+				];
+		}
+		else if (ColumnName == TEXT("Expires"))
+		{
+			return SNew(SBox)
+				.Padding(FMargin(4.0f, 0.0f))
+				.VAlign(VAlign_Center)
+				.HAlign(HAlign_Center)
+				[
+					SNew(STextBlock)
+					.ColorAndOpacity(this, &SCertificateListRow::HandleSelectedColorAndOpacity)
+					.Text(this, &SCertificateListRow::HandleExpiresText)
+				];
+		}
+
+		return SNullWidget::NullWidget;
+	}
+	END_SLATE_FUNCTION_BUILD_OPTIMIZATION
+
+private:
+
+	// Callback for getting the text in the 'Name' column.
+	FText HandleNameText( ) const
+	{
+		return FText::FromString(Certificate->Name);
+	}
+
+	// Callback for getting the status text.
+	FText HandleStatusTextBlockText( ) const
+	{
+		if (Certificate->Status == TEXT("EXPIRED"))
+		{
+			return FText::FromString(TEXT("Expired"));
+		}
+		return FText::FromString(TEXT("Valid"));
+	}
+
+	// Callback for getting the status text color.
+	FSlateColor HandleStatusTextColorAndOpacity( ) const
+	{
+		if (Certificate->Status == TEXT("EXPIRED"))
+		{
+			return FSlateColor(FLinearColor(1.0f, 0.0f, 0.0f));
+		}
+		else if (Certificate->bSelected)
+		{
+			return FSlateColor(FLinearColor(0.0f, 1.0f, 0.0f));
+		}
+		return FSlateColor(FLinearColor(1.0f, 1.0f, 1.0f));
+	}
+
+	// Callback for getting the status text color.
+	FSlateColor HandleSelectedColorAndOpacity( ) const
+	{
+		if (Certificate->bSelected)
+		{
+			return FSlateColor(FLinearColor(0.0f, 1.0f, 0.0f));
+		}
+		return FSlateColor(FLinearColor(1.0f, 1.0f, 1.0f));
+	}
+
+	// Callback for getting the text in the 'Expires' column.
+	FText HandleExpiresText( ) const
+	{
+		return FText::FromString(Certificate->Expires);
+	}
+
+	ECheckBoxState HandleChecked() const
+	{
+		return Certificate->bManuallySelected ? ECheckBoxState::Checked : ECheckBoxState::Unchecked;
+	}
+
+	void HandleCheckStateChanged(ECheckBoxState InState)
+	{
+		Certificate->bManuallySelected = InState == ECheckBoxState::Checked;
+
+		// update the property
+		if (OnCertificateChanged_Handler.IsBound())
+		{
+			OnCertificateChanged_Handler.Execute(Certificate->bManuallySelected ? Certificate->Name : "");
+		}
+
+		// disable any other objects
+		for (int32 Idx = 0; Idx < CertificateList->Num(); ++Idx)
+		{
+			if ((*CertificateList)[Idx] != Certificate && (*CertificateList)[Idx]->bManuallySelected)
+			{
+				(*CertificateList)[Idx]->bManuallySelected = false;
+			}
+		}
+	}
+
+private:
+
+	// Holds the target device service used to populate this row.
+	CertificatePtr Certificate;
+	CertificateListPtr CertificateList;
+	FOnCertificateChanged OnCertificateChanged_Handler;
+};
diff --git a/Engine/Source/Developer/IOS/IOSPlatformEditor/Private/SProvisionListRow.h b/Engine/Source/Developer/IOS/IOSPlatformEditor/Private/SProvisionListRow.h
new file mode 100644
--- /dev/null
+++ b/Engine/Source/Developer/IOS/IOSPlatformEditor/Private/SProvisionListRow.h
@@ -0,0 +1,244 @@
+// Copyright 1998-2017 Epic Games, Inc. All Rights Reserved.
+
+#pragma once
+
+#include "CoreMinimal.h"
+#include "Styling/SlateColor.h"
+#include "Widgets/SWidget.h"
+#include "Layout/Margin.h"
+#include "Widgets/SNullWidget.h"
+#include "Widgets/Views/STableViewBase.h"
+#include "Widgets/DeclarativeSyntaxSupport.h"
+#include "Styling/SlateTypes.h"
+#include "Misc/Paths.h"
+#include "Brushes/SlateNoResource.h"
+#include "Widgets/Layout/SBox.h"
+#include "Widgets/Views/STableRow.h"
+#include "IOSTargetSettingsCustomization.h"
+#include "Widgets/Text/STextBlock.h"
+#include "Widgets/Input/SCheckBox.h"
+#include "Widgets/Views/SListView.h"
+#include "Brushes/SlateImageBrush.h"
+#include "SlateOptMacros.h"
+
+#define IMAGE_BRUSH( RelativePath, ... ) FSlateImageBrush( FPaths::EngineContentDir() / TEXT("Editor/Slate") / RelativePath + TEXT(".png"), __VA_ARGS__ )
+
+DECLARE_DELEGATE_OneParam(FOnProvisionChanged, FString);
+
+/**
+ * Implements a row widget for the provision list view.
+ */
+class SProvisionListRow
+	: public SMultiColumnTableRow<ProvisionPtr>
+{
+public:
+
+	SLATE_BEGIN_ARGS(SProvisionListRow) { }
+		SLATE_ARGUMENT(ProvisionPtr, Provision)
+		SLATE_ARGUMENT(ProvisionListPtr, ProvisionList)
+		SLATE_EVENT(FOnProvisionChanged, OnProvisionChanged)
+	SLATE_END_ARGS()
+
+public:
+
+	/**
+	 * Constructs the widget.
+	 *
+	 * @param InArgs The arguments.
+	 */
+	void Construct( const FArguments& InArgs, const TSharedRef<STableViewBase>& InOwnerTableView )
+	{
+		Provision = InArgs._Provision;
+		ProvisionList = InArgs._ProvisionList;
+		OnProvisionChanged_Handler = InArgs._OnProvisionChanged;
+
+		SMultiColumnTableRow<ProvisionPtr>::Construct(FSuperRowType::FArguments(), InOwnerTableView);
+
+		/* Set images for various SCheckBox states ... */
+		if (!bInitialized)
+		{
+			ProvisionCheckBoxStyle = FCheckBoxStyle()
+				.SetCheckBoxType(ESlateCheckBoxType::CheckBox)
+				.SetUncheckedImage( FSlateNoResource() )
+				.SetUncheckedHoveredImage( FSlateNoResource() )
+				.SetUncheckedPressedImage( FSlateNoResource() )
+				.SetCheckedImage( IMAGE_BRUSH( "Automation/Success", FVector2D(16.0f, 16.0f) ) )
+				.SetCheckedHoveredImage( IMAGE_BRUSH( "Automation/Success", FVector2D(16.0f, 16.0f), FLinearColor( 0.5f, 0.5f, 0.5f ) ) )
+				.SetCheckedPressedImage( IMAGE_BRUSH( "Automation/Success", FVector2D(16.0f, 16.0f) ) )
+				.SetUndeterminedImage(FSlateNoResource() )
+				.SetUndeterminedHoveredImage( FSlateNoResource() )
+				.SetUndeterminedPressedImage( FSlateNoResource() );
+
+			bInitialized = true;
+		}
+	}
+
+public:
+
+	/**
+	 * Generates the widget for the specified column.
+	 *
+	 * @param ColumnName The name of the column to generate the widget for.
+	 * @return The widget.
+	 */
+	BEGIN_SLATE_FUNCTION_BUILD_OPTIMIZATION
+	virtual TSharedRef<SWidget> GenerateWidgetForColumn( const FName& ColumnName ) override
+	{
+		if (ColumnName == TEXT("Selected"))
+		{
+			return SNew(SCheckBox)
+				.IsChecked(this, &SProvisionListRow::HandleChecked)
+				.OnCheckStateChanged(this, &SProvisionListRow::HandleCheckStateChanged);
+		}
+		else if (ColumnName == TEXT("Name"))
+		{
+			return SNew(SBox)
+				.Padding(FMargin(4.0f, 0.0f))
+				.VAlign(VAlign_Center)
+				[
+					SNew(STextBlock)
+						.ColorAndOpacity(this, &SProvisionListRow::HandleSelectedColorAndOpacity)
+						.Text(this, &SProvisionListRow::HandleNameText)
+				];
+		}
+		else if (ColumnName == TEXT("File"))
+		{
+			return SNew(SBox)
+				.Padding(FMargin(4.0f, 0.0f))
+				.VAlign(VAlign_Center)
+				[
+					SNew(STextBlock)
+					.ColorAndOpacity(this, &SProvisionListRow::HandleSelectedColorAndOpacity)
+					.Text(this, &SProvisionListRow::HandleFileText)
+				];
+		}
+		else if (ColumnName == TEXT("Status"))
+		{
+			return SNew(SBox)
+				.Padding(FMargin(4.0f, 0.0f))
+				.VAlign(VAlign_Center)
+				[
+					SNew(STextBlock)
+						.ColorAndOpacity(this, &SProvisionListRow::HandleStatusTextColorAndOpacity)
+						.Text(this, &SProvisionListRow::HandleStatusTextBlockText)
+				];
+		}
+		else if (ColumnName == TEXT("Distribution"))
+		{
+			return SNew(SBox)
+				.Padding(FMargin(4.0f, 0.0f))
+				.VAlign(VAlign_Center)
+				.HAlign(HAlign_Center)
+				[
+					SNew(SCheckBox)
+						.IsChecked(this, &SProvisionListRow::HandleDistribution)
+						.Style(&ProvisionCheckBoxStyle)
+				];
+		}
+
+		return SNullWidget::NullWidget;
+	}
+	END_SLATE_FUNCTION_BUILD_OPTIMIZATION
+
+
+private:
+
+	// Callback for getting the text in the 'Name' column.
+	FText HandleNameText( ) const
+	{
+		return FText::FromString(Provision->Name);
+	}
+
+	// Callback for getting the text in the 'File' column.
+	FText HandleFileText( ) const
+	{
+		return FText::FromString(Provision->FileName);
+	}
+
+	// Callback for getting the status text.
+	FText HandleStatusTextBlockText( ) const
+	{
+		if (Provision->Status == TEXT("NO_MATCH"))
+		{
+			return FText::FromString(TEXT("Identifier Not Matched"));
+		}
+		else if (Provision->Status == TEXT("NO_CERT"))
+		{
+			return FText::FromString(TEXT("No Valid Certificate Found"));
+		}
+		else if (Provision->Status == TEXT("EXPIRED"))
+		{
+			return FText::FromString(TEXT("Expired"));
+		}
+		return FText::FromString(TEXT("Valid"));
+	}
+
+	// Callback for getting the status text color.
+	FSlateColor HandleStatusTextColorAndOpacity( ) const
+	{
+		if (Provision->Status == TEXT("NO_MATCH"))
+		{
+			return FSlateColor(FLinearColor(1.0f, 1.0f, 0.0f));
+		}
+		else if (Provision->Status == TEXT("NO_CERT") || Provision->Status == TEXT("EXPIRED"))
+		{
+			return FSlateColor(FLinearColor(1.0f, 0.0f, 0.0f));
+		}
+		else if (Provision->bSelected)
+		{
+			return FSlateColor(FLinearColor(0.0f, 1.0f, 0.0f));
+		}
+		return FSlateColor(FLinearColor(1.0f, 1.0f, 1.0f));
+	}
+
+	// Callback for getting the status text color.
+	FSlateColor HandleSelectedColorAndOpacity( ) const
+	{
+		if (Provision->bSelected)
+		{
+			return FSlateColor(FLinearColor(0.0f, 1.0f, 0.0f));
+		}
+		return FSlateColor(FLinearColor(1.0f, 1.0f, 1.0f));
+	}
+
+	// Callback to determine distribution
+	ECheckBoxState HandleDistribution( ) const
+	{
+		return Provision->bDistribution ? ECheckBoxState::Checked : ECheckBoxState::Unchecked;
+	}
+
+	ECheckBoxState HandleChecked() const
+	{
+		return Provision->bManuallySelected ? ECheckBoxState::Checked : ECheckBoxState::Unchecked;
+	}
+
+	void HandleCheckStateChanged(ECheckBoxState InState)
+	{
+		Provision->bManuallySelected = InState == ECheckBoxState::Checked;
+
+		// update the property
+		if (OnProvisionChanged_Handler.IsBound())
+		{
+			OnProvisionChanged_Handler.Execute(Provision->bManuallySelected ? Provision->FileName : "");
+		}
+
+		// disable any other objects
+		for (int32 Idx = 0; Idx < ProvisionList->Num(); ++Idx)
+		{
+			if ((*ProvisionList)[Idx] != Provision && (*ProvisionList)[Idx]->bManuallySelected)
+			{
+				(*ProvisionList)[Idx]->bManuallySelected = false;
+			}
+		}
+	}
+
+private:
+
+	// Holds the target device service used to populate this row.
+	ProvisionPtr Provision;
+	ProvisionListPtr ProvisionList;
+	FOnProvisionChanged OnProvisionChanged_Handler;
+
+	static bool bInitialized;
+	static FCheckBoxStyle ProvisionCheckBoxStyle;
+};
diff --git a/Engine/Source/Developer/IOS/IOSTargetPlatform/IOSTargetPlatform.Build.cs b/Engine/Source/Developer/IOS/IOSTargetPlatform/IOSTargetPlatform.Build.cs
new file mode 100644
--- /dev/null
+++ b/Engine/Source/Developer/IOS/IOSTargetPlatform/IOSTargetPlatform.Build.cs
@@ -0,0 +1,53 @@
+// Copyright 1998-2017 Epic Games, Inc. All Rights Reserved.
+
+using UnrealBuildTool;
+
+public class IOSTargetPlatform : ModuleRules
+{
+	public IOSTargetPlatform(ReadOnlyTargetRules Target) : base(Target)
+	{
+		BinariesSubFolder = "IOS";
+		
+		PrivateDependencyModuleNames.AddRange(
+			new string[]
+			{
+				"Core",
+				"CoreUObject",
+				"TargetPlatform",
+				"DesktopPlatform",
+				"LaunchDaemonMessages",
+				"Projects"
+			}
+		);
+
+		PrivateIncludePathModuleNames.AddRange(
+			new string[] {
+				"Messaging",
+				"TargetDeviceServices",
+			}
+		);
+
+		PlatformSpecificDynamicallyLoadedModuleNames.Add("LaunchDaemonMessages");
+
+		PrivateIncludePaths.Add("Developer/IOS/IOSTargetPlatform/Private");
+
+		//This is somehow necessary for getting iOS to build on, at least, windows. It seems like the target platform is included for cooking, and thus it requirtes a bunch of other info.
+		PublicIncludePaths.AddRange(
+			new string[]
+			{
+				"Runtime/Core/Public/Apple",
+				"Runtime/Core/Public/IOS",
+			}
+		);
+
+		if (UEBuildConfiguration.bCompileAgainstEngine)
+		{
+			PrivateDependencyModuleNames.Add("Engine");
+		}
+
+		if (Target.Platform == UnrealTargetPlatform.Mac)
+		{
+            PublicAdditionalLibraries.Add("/System/Library/PrivateFrameworks/MobileDevice.framework/Versions/Current/MobileDevice");
+		}
+	}
+}
diff --git a/Engine/Source/Developer/IOS/IOSTargetPlatform/Private/IOSDeviceHelper.h b/Engine/Source/Developer/IOS/IOSTargetPlatform/Private/IOSDeviceHelper.h
new file mode 100644
--- /dev/null
+++ b/Engine/Source/Developer/IOS/IOSTargetPlatform/Private/IOSDeviceHelper.h
@@ -0,0 +1,67 @@
+// Copyright 1998-2017 Epic Games, Inc. All Rights Reserved.
+
+#pragma once
+
+/**
+ * Type definition for shared pointers to instances of FIOSDeviceHelper
+ */
+typedef TSharedPtr<class FIOSDeviceHelper, ESPMode::ThreadSafe> FIOSDeviceHelperPtr;
+
+/**
+ * Type definition for shared references to instances of FIOSDeviceHelper
+ */
+typedef TSharedRef<class FIOSDeviceHelper, ESPMode::ThreadSafe> FIOSDeviceHelperRef;
+
+/**
+ * Delegate type for devices being connected or disconnected from the machine
+ *
+ * The first parameter is newly added or removed device
+ */
+DECLARE_MULTICAST_DELEGATE_OneParam(FOnDeviceConnectEvent, const FIOSLaunchDaemonPong&)
+
+/**
+ * Implemented the iOS device helper class
+ */
+class FIOSDeviceHelper
+{
+public:
+    static void Initialize(bool bIsTVOS);
+    
+    /**
+     * Returns a delegate that is executed when a device is connected
+     *
+     * @return the delegate
+     */
+    static FOnDeviceConnectEvent& OnDeviceConnected()
+    {
+        static FOnDeviceConnectEvent DeviceConnectedDelegate;
+        return DeviceConnectedDelegate;
+    }
+    
+    /**
+     * Returns a delegate the is executed when a device is disconnected
+     *
+     * @return the delegate
+     */
+    static FOnDeviceConnectEvent& OnDeviceDisconnected()
+    {
+        static FOnDeviceConnectEvent DeviceDisconnectedDelegate;
+        return DeviceDisconnectedDelegate;
+    }
+    
+    /**
+     * Installs an ipa on to the device
+     */
+    static bool InstallIPAOnDevice(const FTargetDeviceId& DeviceId, const FString& IPAPath);
+    
+	/**
+	 * Suspends/Enables the device connect/disconnect thread
+	 */
+	static void EnableDeviceCheck(bool OnOff);
+
+private:
+    static void DeviceCallback(void*);
+    static void DoDeviceConnect(void*);
+    static void DoDeviceDisconnect(void*);
+	static bool MessageTickDelegate(float);
+};
diff --git a/Engine/Source/Developer/IOS/IOSTargetPlatform/Private/IOSTargetDevice.cpp b/Engine/Source/Developer/IOS/IOSTargetPlatform/Private/IOSTargetDevice.cpp
new file mode 100644
--- /dev/null
+++ b/Engine/Source/Developer/IOS/IOSTargetPlatform/Private/IOSTargetDevice.cpp
@@ -0,0 +1,170 @@
+// Copyright 1998-2017 Epic Games, Inc. All Rights Reserved.
+
+#include "IOSTargetDevice.h"
+#include "IOSMessageProtocol.h"
+
+bool FIOSTargetDevice::Connect()
+{
+	// @todo zombie - Probably need to write a specific ConnectTo(IpAddr) function for setting up a RemoteEndpoint for talking to the Daemon
+	// Returning true since, if this exists, a device exists.
+
+	return true;
+}
+
+bool FIOSTargetDevice::Deploy(const FString& SourceFolder, FString& OutAppId)
+{
+	return false;
+}
+
+void FIOSTargetDevice::Disconnect()
+{
+}
+
+int32 FIOSTargetDevice::GetProcessSnapshot(TArray<FTargetDeviceProcessInfo>& OutProcessInfos)
+{
+	return 0;
+}
+
+ETargetDeviceTypes FIOSTargetDevice::GetDeviceType() const
+{
+	return DeviceType;
+}
+
+FTargetDeviceId FIOSTargetDevice::GetId() const
+{
+	return DeviceId;
+}
+
+FString FIOSTargetDevice::GetName() const
+{
+	return DeviceName;
+}
+
+FString FIOSTargetDevice::GetOperatingSystemName()
+{
+	return TargetPlatform.PlatformName();
+}
+
+const class ITargetPlatform& FIOSTargetDevice::GetTargetPlatform() const
+{
+	return TargetPlatform;
+}
+
+bool FIOSTargetDevice::IsConnected()
+{
+	return true;
+}
+
+bool FIOSTargetDevice::IsDefault() const
+{
+	return true;
+}
+	
+bool FIOSTargetDevice::Launch(const FString& InAppId, EBuildConfigurations::Type InBuildConfiguration, EBuildTargets::Type BuildTarget, const FString& Params, uint32* OutProcessId)
+{
+#if !PLATFORM_MAC
+	MessageEndpoint->Send(new FIOSLaunchDaemonLaunchApp(InAppId, Params), DeviceEndpoint);
+	return true;
+#else
+	//Set return to false on Mac, since we could not find a way to do remote deploy/launch.
+	return false;
+#endif // !PLATFORM_MAC
+}
+
+bool FIOSTargetDevice::PowerOff(bool Force)
+{
+	// @todo zombie - Supported by the Daemon?
+
+	return false;
+}
+
+bool FIOSTargetDevice::PowerOn()
+{
+	// @todo zombie - Supported by the Daemon?
+
+	return false;
+}
+
+bool FIOSTargetDevice::Reboot(bool bReconnect)
+{
+	// @todo zombie - Supported by the Daemon?
+
+	return false;
+}
+
+bool FIOSTargetDevice::Run(const FString& ExecutablePath, const FString& Params, uint32* OutProcessId)
+{
+#if !PLATFORM_MAC
+	// The executable path usually looks something like this: directory/<gamename>.stub
+	// We just need <gamename>, so strip that out.
+	int32 LastPeriodPos = ExecutablePath.Find( TEXT("."), ESearchCase::IgnoreCase, ESearchDir::FromEnd);
+	int32 LastBackslashPos = ExecutablePath.Find( TEXT("\\"), ESearchCase::IgnoreCase, ESearchDir::FromEnd);
+	int32 LastSlashPos = ExecutablePath.Find( TEXT("/"), ESearchCase::IgnoreCase, ESearchDir::FromEnd);
+	int32 TrimPos = LastBackslashPos > LastSlashPos ? LastBackslashPos : LastSlashPos;
+	if ( TrimPos > LastPeriodPos )
+	{
+		// Ignore any periods in the path before the last "/" or "\"
+		LastPeriodPos = ExecutablePath.Len() - 1;
+	}
+
+	if ( TrimPos == INDEX_NONE )
+	{
+		TrimPos = 0; // take the whole string from the beginning
+	}
+	else
+	{
+		// increment to one character beyond the slash
+		TrimPos++;
+	}
+
+	int32 Count = LastPeriodPos - TrimPos;
+	FString NewAppId = ExecutablePath.Mid(TrimPos, Count); // remove the ".stub" and the proceeding directory to get the game name
+
+	SetAppId(NewAppId);
+	MessageEndpoint->Send(new FIOSLaunchDaemonLaunchApp(AppId, Params), DeviceEndpoint);
+	return true;
+#else
+	//Set return to false on Mac, since we could not find a way to do remote deploy/launch.
+	return false;
+#endif // !PLATFORM_MAC
+}
+
+bool FIOSTargetDevice::SupportsFeature(ETargetDeviceFeatures Feature) const
+{
+	switch (Feature)
+	{
+	case ETargetDeviceFeatures::Reboot:
+		return bCanReboot;
+
+	case ETargetDeviceFeatures::PowerOn:
+		return bCanPowerOn;
+
+	case ETargetDeviceFeatures::PowerOff:
+		return bCanPowerOff;
+
+	case ETargetDeviceFeatures::ProcessSnapshot:
+		return false;
+
+	default:
+		return false;
+	}
+}
+
+bool FIOSTargetDevice::SupportsSdkVersion(const FString& VersionString) const
+{
+	return true;
+}
+
+bool FIOSTargetDevice::TerminateProcess(const int64 ProcessId)
+{
+	return false;
+}
+
+void FIOSTargetDevice::SetUserCredentials(const FString& UserName, const FString& UserPassword)
+{
+}
+
+bool FIOSTargetDevice::GetUserCredentials(FString& OutUserName, FString& OutUserPassword)
+{
+	return false;
+}
diff --git a/Engine/Source/Developer/IOS/IOSTargetPlatform/Private/IOSTargetDevice.h b/Engine/Source/Developer/IOS/IOSTargetPlatform/Private/IOSTargetDevice.h
new file mode 100644
--- /dev/null
+++ b/Engine/Source/Developer/IOS/IOSTargetPlatform/Private/IOSTargetDevice.h
@@ -0,0 +1,184 @@
+// Copyright 1998-2017 Epic Games, Inc. All Rights Reserved.
+
+/*=============================================================================
+	IOSTargetDevice.h: Defines an IOS device for communication
+=============================================================================*/
+
+#pragma once
+
+#include "CoreMinimal.h"
+#include "Interfaces/ITargetDevice.h"
+#include "Interfaces/ITargetPlatform.h"
+#include "HAL/PlatformProcess.h"
+#include "IMessageContext.h"
+#include "Helpers/MessageEndpoint.h"
+#include "Helpers/MessageEndpointBuilder.h"
+
+
+/**
+ * Type definition for shared pointers to instances of FIOSTargetDevice.
+ */
+typedef TSharedPtr<class FIOSTargetDevice, ESPMode::ThreadSafe> FIOSTargetDevicePtr;
+
+/**
+ * Type definition for shared references to instances of FIOSTargetDevice.
+ */
+typedef TSharedRef<class FIOSTargetDevice, ESPMode::ThreadSafe> FIOSTargetDeviceRef;
+
+
+/**
+ * Implements an iOS target device.
+ */
+class FIOSTargetDevice : public ITargetDevice
+{
+public:
+	FIOSTargetDevice (const ITargetPlatform& InTargetPlatform)
+		: TargetPlatform(InTargetPlatform)
+		, DeviceEndpoint()
+		, AppId()
+		, bCanReboot(false)
+		, bCanPowerOn(false)
+		, bCanPowerOff(false)
+		, DeviceType(ETargetDeviceTypes::Indeterminate)
+	{
+		DeviceId = FTargetDeviceId(TargetPlatform.PlatformName(), FPlatformProcess::ComputerName());
+		DeviceName = FPlatformProcess::ComputerName();
+		MessageEndpoint = FMessageEndpoint::Builder("FIOSTargetDevice").Build();
+	}
+public:
+	// ITargetDevice interface
+	virtual bool Connect() override;
+	virtual bool Deploy(const FString& SourceFolder, FString& OutAppId) override;
+	virtual void Disconnect() override;
+	virtual int32 GetProcessSnapshot(TArray<FTargetDeviceProcessInfo>& OutProcessInfos) override;
+	virtual ETargetDeviceTypes GetDeviceType() const override;
+	virtual FTargetDeviceId GetId() const override;
+	virtual FString GetName() const override;
+	virtual FString GetOperatingSystemName() override;
+	virtual const class ITargetPlatform& GetTargetPlatform() const override;
+	virtual bool IsConnected() override;
+	virtual bool IsDefault() const override;
+	virtual bool Launch(const FString& InAppId, EBuildConfigurations::Type InBuildConfiguration, EBuildTargets::Type BuildTarget, const FString& Params, uint32* OutProcessId) override;
+	virtual bool PowerOff(bool Force) override;
+	virtual bool PowerOn() override;
+	virtual bool Reboot(bool bReconnect = false) override;
+	virtual bool Run(const FString& ExecutablePath, const FString& Params, uint32* OutProcessId) override;
+	virtual bool SupportsFeature(ETargetDeviceFeatures Feature) const;
+	virtual bool SupportsSdkVersion(const FString& VersionString) const override;
+	virtual bool TerminateProcess(const int64 ProcessId) override;
+	virtual void SetUserCredentials(const FString& UserName, const FString& UserPassword) override;
+	virtual bool GetUserCredentials(FString& OutUserName, FString& OutUserPassword) override;
+	// End of ITargetDevice interface
+
+public:
+	/** Timeout check for removing stale devices */
+	FDateTime LastPinged;
+
+private:
+
+	/** The current status of this device. */
+//	ETargetDeviceStatus::Type Status;
+	/** Holds a reference to the device's target platform. */
+	const ITargetPlatform& TargetPlatform;
+	/** Contains the address of the remote device */
+	FMessageAddress DeviceEndpoint;
+	/** MessageEndpoint for communicating with remote device */
+	FMessageEndpointPtr MessageEndpoint;
+	/** Contains the current AppID/GameName for Deployment/launching. */
+	FString AppId;
+	/** Containts the build configuration of the app to deploy */
+	EBuildConfigurations::Type BuildConfiguration;
+	/** Lets us know whether the thing is a sim device or a physical device. */
+	bool bIsSimulated;
+
+
+private:
+
+	/** Remote rebootable */
+	bool bCanReboot;
+	/** Remote bootable */
+	bool bCanPowerOn;
+	/** Remote shutdown-able */
+	bool bCanPowerOff;
+	/** Id of device */
+	FTargetDeviceId DeviceId;
+	/** Name of device */
+	FString DeviceName;
+	/** Type of device */
+	ETargetDeviceTypes DeviceType;
+
+public:
+
+	void SetFeature(ETargetDeviceFeatures InFeature, bool bFlag)
+	{
+		if(InFeature == ETargetDeviceFeatures::Reboot)
+		{
+			bCanReboot = bFlag;
+		}
+		else if(InFeature == ETargetDeviceFeatures::PowerOn)
+		{
+			bCanPowerOn = bFlag;
+		}
+		else if(InFeature == ETargetDeviceFeatures::PowerOff)
+		{
+			bCanPowerOff = bFlag;
+		}
+	}
+
+	/** Sets device id */
+	void SetDeviceId(const FTargetDeviceId InDeviceId)
+	{
+		DeviceId = InDeviceId;
+	}
+
+	/** Sets the name of the device */
+	void SetDeviceName(const FString InDeviceName)
+	{
+		DeviceName = InDeviceName;
+	}
+
+	/** Sets the type of the device */
+	void SetDeviceType(const FString InDeviceTypeString)
+	{
+		if(InDeviceTypeString == TEXT("Browser"))
+		{
+			DeviceType = ETargetDeviceTypes::Browser;
+		}
+		else if(InDeviceTypeString == TEXT("Console"))
+		{
+			DeviceType = ETargetDeviceTypes::Console;
+		}
+		else if(InDeviceTypeString == TEXT("Phone"))
+		{
+			DeviceType = ETargetDeviceTypes::Phone;
+		}
+		else if(InDeviceTypeString == TEXT("Tablet"))
+		{
+			DeviceType = ETargetDeviceTypes::Tablet;
+		}
+		else
+		{
+			DeviceType = ETargetDeviceTypes::Indeterminate;
+		}
+	}
+
+	void SetDeviceEndpoint(const FMessageAddress& DeviceAddress)
+	{
+		DeviceEndpoint = DeviceAddress;
+	}
+
+	void SetAppId(const FString& GameName)
+	{
+		AppId = GameName;
+	}
+
+	void SetAppConfiguration(EBuildConfigurations::Type Configuration)
+	{
+		BuildConfiguration = Configuration;
+	}
+
+	void SetIsSimulated(bool IsSimulated)
+	{
+		bIsSimulated = IsSimulated;
+	}
+};
diff --git a/Engine/Source/Developer/IOS/IOSTargetPlatform/Private/IOSTargetPlatform.cpp b/Engine/Source/Developer/IOS/IOSTargetPlatform/Private/IOSTargetPlatform.cpp
new file mode 100644
--- /dev/null
+++ b/Engine/Source/Developer/IOS/IOSTargetPlatform/Private/IOSTargetPlatform.cpp
@@ -0,0 +1,588 @@
+// Copyright 1998-2017 Epic Games, Inc. All Rights Reserved.
+
+/*=============================================================================
+	IOSTargetPlatform.cpp: Implements the FIOSTargetPlatform class.
+=============================================================================*/
+
+#include "IOSTargetPlatform.h"
+#include "IProjectManager.h"
+#include "InstalledPlatformInfo.h"
+#include "HAL/FileManager.h"
+#include "Misc/Paths.h"
+#include "Misc/App.h"
+#include "Misc/MonitoredProcess.h"
+#if PLATFORM_WINDOWS
+#include "WindowsHWrapper.h"
+#endif
+#if WITH_ENGINE
+#include "TextureResource.h"
+#endif
+
+/* FIOSTargetPlatform structors
+ *****************************************************************************/
+
+FIOSTargetPlatform::FIOSTargetPlatform(bool bInIsTVOS)
+	: bIsTVOS(bInIsTVOS)
+{
+    if (bIsTVOS)
+    {
+        this->PlatformInfo = PlatformInfo::FindPlatformInfo("TVOS");
+    }
+#if WITH_ENGINE
+	FConfigCacheIni::LoadLocalIniFile(EngineSettings, TEXT("Engine"), true, *PlatformName());
+	TextureLODSettings = nullptr; // TextureLODSettings are registered by the device profile.
+	StaticMeshLODSettings.Initialize(EngineSettings);
+#endif // #if WITH_ENGINE
+
+	// Initialize Ticker for device discovery
+	TickDelegate = FTickerDelegate::CreateRaw(this, &FIOSTargetPlatform::HandleTicker);
+	TickDelegateHandle = FTicker::GetCoreTicker().AddTicker(TickDelegate, 10.0f);
+	
+	// initialize the connected device detector
+	DeviceHelper.OnDeviceConnected().AddRaw(this, &FIOSTargetPlatform::HandleDeviceConnected);
+	DeviceHelper.OnDeviceDisconnected().AddRaw(this, &FIOSTargetPlatform::HandleDeviceDisconnected);
+	DeviceHelper.Initialize(bIsTVOS);
+}
+
+
+FIOSTargetPlatform::~FIOSTargetPlatform()
+{
+	FTicker::GetCoreTicker().RemoveTicker(TickDelegateHandle);
+}
+
+
+/* ITargetPlatform interface
+ *****************************************************************************/
+
+void FIOSTargetPlatform::EnableDeviceCheck(bool OnOff)
+{
+	FIOSDeviceHelper::EnableDeviceCheck(OnOff);
+}
+
+void FIOSTargetPlatform::GetAllDevices( TArray<ITargetDevicePtr>& OutDevices ) const
+{
+	OutDevices.Reset();
+
+	for (auto Iter = Devices.CreateConstIterator(); Iter; ++Iter)
+	{
+		OutDevices.Add(Iter.Value());
+	}
+}
+
+
+ITargetDevicePtr FIOSTargetPlatform::GetDefaultDevice() const
+{
+	if (Devices.Num() > 0)
+	{
+		// first device is the default
+		auto Iter = Devices.CreateConstIterator();
+		if(Iter)
+		{
+			return Iter.Value();
+		}
+	}
+
+	return NULL;
+}
+
+
+ITargetDevicePtr FIOSTargetPlatform::GetDevice( const FTargetDeviceId& DeviceId )
+{
+	return Devices.FindRef(DeviceId);
+}
+
+
+bool FIOSTargetPlatform::IsSdkInstalled(bool bProjectHasCode, FString& OutTutorialPath) const
+{
+#if PLATFORM_MAC
+	OutTutorialPath = FString("Shared/Tutorials/InstallingXCodeTutorial");
+	bool biOSSDKInstalled = IFileManager::Get().DirectoryExists(TEXT("/Applications/Xcode.app/Contents/Developer/Platforms/iPhoneOS.platform"));
+#else
+	OutTutorialPath = FString("/Engine/Tutorial/Mobile/InstallingiTunesTutorial.InstallingiTunesTutorial");
+
+	// On windows we check if itunes is installed - Perhaps someday make this its own check instead of piggy packing on the SDK check which will create a unintuitive error message when it fails
+
+	// The logic here is to assume the correct Apple dll does not exist and then check the various locations it could be in, setting this to true when it is found
+	// Code is structured for clarity not performance
+	// See Engine\Source\Programs\IOS\MobileDeviceInterface\MobileDevice.cs for reference
+	bool biOSSDKInstalled = false; 
+
+	HKEY hKey;
+	TCHAR dllPath[256];
+	unsigned long pathSize = 256;
+	
+	// Add future version checks here
+
+	// Check for iTunes 12
+	if(!biOSSDKInstalled
+		&& RegOpenKeyEx(HKEY_LOCAL_MACHINE, TEXT("SOFTWARE\\Wow6432Node\\Apple Inc.\\Apple Mobile Device Support\\Shared"), 0, KEY_READ, &hKey) == ERROR_SUCCESS
+		&& RegQueryValueEx(hKey, TEXT("MobileDeviceDLL"), 0, NULL, (BYTE*)dllPath, &pathSize) == ERROR_SUCCESS
+		&&  IFileManager::Get().FileSize(*FString(dllPath)) != INDEX_NONE)
+		{
+		biOSSDKInstalled = true;
+	}
+	
+	// Check for iTunes 11
+	if(!biOSSDKInstalled
+		&& RegOpenKeyEx(HKEY_LOCAL_MACHINE, TEXT("SOFTWARE\\Wow6432Node\\Apple Inc.\\Apple Mobile Device Support\\Shared"), 0, KEY_READ, &hKey) == ERROR_SUCCESS
+		&& RegQueryValueEx(hKey, TEXT("iTunesMobileDeviceDLL"), 0, NULL, (BYTE*)dllPath, &pathSize) == ERROR_SUCCESS
+		&&  IFileManager::Get().FileSize(*FString(dllPath)) != INDEX_NONE)
+			{
+		biOSSDKInstalled = true;
+	}
+
+#endif
+	return biOSSDKInstalled;
+}
+
+static FString OutputMessage;
+static void OnOutput(FString Message)
+{
+	OutputMessage += Message;
+	UE_LOG(LogTemp, Display, TEXT("%s\n"), *Message);
+}
+
+int32 FIOSTargetPlatform::CheckRequirements(const FString& ProjectPath, bool bProjectHasCode, FString& OutTutorialPath) const
+{
+	int32 bReadyToBuild = ETargetPlatformReadyStatus::Ready; // @todo How do we check that the iOS SDK is installed when building from Windows? Is that even possible?
+	if (!IsSdkInstalled(bProjectHasCode, OutTutorialPath))
+	{
+		bReadyToBuild |= ETargetPlatformReadyStatus::SDKNotFound;
+	}
+#if PLATFORM_MAC
+	OutTutorialPath = FString("/Engine/Tutorial/Installation/InstallingXCodeTutorial.InstallingXCodeTutorial");
+    // shell to certtool
+#else
+	if (!FInstalledPlatformInfo::Get().IsValidPlatform(GetPlatformInfo().BinaryFolderName, EProjectType::Code))
+	{
+		if (bProjectHasCode)
+		{
+			OutTutorialPath = FString("/Engine/Tutorial/Mobile/iOSonPCRestrictions.iOSonPCRestrictions");
+			bReadyToBuild |= ETargetPlatformReadyStatus::CodeUnsupported;
+		}
+		if (IProjectManager::Get().IsNonDefaultPluginEnabled())
+		{
+			OutTutorialPath = FString("/Engine/Tutorial/Mobile/iOSonPCValidPlugins.iOSonPCValidPlugins");
+			bReadyToBuild |= ETargetPlatformReadyStatus::PluginsUnsupported;
+		}
+	}
+#endif
+
+	// shell to IPP and get the status of the provision and cert
+
+	bool bForDistribtion = false;
+	GConfig->GetBool(TEXT("/Script/UnrealEd.ProjectPackagingSettings"), TEXT("ForDistribution"), bForDistribtion, GGameIni);
+
+	FString BundleIdentifier;
+	GConfig->GetString(TEXT("/Script/IOSRuntimeSettings.IOSRuntimeSettings"), TEXT("BundleIdentifier"), BundleIdentifier, GEngineIni);
+	BundleIdentifier = BundleIdentifier.Replace(TEXT("[PROJECT_NAME]"), FApp::GetGameName());
+	BundleIdentifier = BundleIdentifier.Replace(TEXT("_"), TEXT(""));
+#if PLATFORM_MAC
+    FString CmdExe = TEXT("/bin/sh");
+    FString ScriptPath = FPaths::ConvertRelativePathToFull(FPaths::EngineDir() / TEXT("Build/BatchFiles/Mac/RunMono.sh"));
+    FString IPPPath = FPaths::ConvertRelativePathToFull(FPaths::EngineDir() / TEXT("Binaries/DotNET/IOS/IPhonePackager.exe"));
+	FString CommandLine = FString::Printf(TEXT("\"%s\" \"%s\" Validate Engine -project \"%s\" -bundlename \"%s\" %s"), *ScriptPath, *IPPPath, *ProjectPath, *(BundleIdentifier), (bForDistribtion ? TEXT("-distribution") : TEXT("")) );
+#else
+	FString CmdExe = FPaths::ConvertRelativePathToFull(FPaths::EngineDir() / TEXT("Binaries/DotNET/IOS/IPhonePackager.exe"));
+	FString CommandLine = FString::Printf(TEXT("Validate Engine -project \"%s\" -bundlename \"%s\" %s"), *ProjectPath, *(BundleIdentifier), (bForDistribtion ? TEXT("-distribution") : TEXT("")) );
+	FString RemoteServerName;
+	GConfig->GetString(TEXT("/Script/IOSRuntimeSettings.IOSRuntimeSettings"), TEXT("RemoteServerName"), RemoteServerName, GEngineIni);
+	if (RemoteServerName.Len() == 0)
+	{
+		bReadyToBuild |= ETargetPlatformReadyStatus::RemoveServerNameEmpty;
+	}
+
+#endif
+	TSharedPtr<FMonitoredProcess> IPPProcess = MakeShareable(new FMonitoredProcess(CmdExe, CommandLine, true));
+	OutputMessage = TEXT("");
+	IPPProcess->OnOutput().BindStatic(&OnOutput);
+	IPPProcess->Launch();
+	while(IPPProcess->Update())
+	{
+		FPlatformProcess::Sleep(0.01f);
+	}
+	int RetCode = IPPProcess->GetReturnCode();
+    UE_LOG(LogTemp, Display, TEXT("%s"), *OutputMessage);
+	if (RetCode == 14)
+	{
+		OutTutorialPath = FString("/Engine/Tutorial/Mobile/CreatingInfoPlist.CreatingInfoPlist");
+		bReadyToBuild |= ETargetPlatformReadyStatus::ManifestNotFound;
+	}
+	else if (RetCode == 13)
+	{
+		OutTutorialPath = FString("/Engine/Tutorial/Mobile/CreatingSigningCertAndProvisionTutorial.CreatingSigningCertAndProvisionTutorial");
+		bReadyToBuild |= ETargetPlatformReadyStatus::SigningKeyNotFound;
+		bReadyToBuild |= ETargetPlatformReadyStatus::ProvisionNotFound;
+	}
+	else if (RetCode == 12)
+	{
+		OutTutorialPath = FString("/Engine/Tutorial/Mobile/CreatingSigningCertAndProvisionTutorial.CreatingSigningCertAndProvisionTutorial");
+		bReadyToBuild |= ETargetPlatformReadyStatus::SigningKeyNotFound;
+	}
+	else if (RetCode == 11)
+	{
+		OutTutorialPath = FString("/Engine/Tutorial/Mobile/CreatingSigningCertAndProvisionTutorial.CreatingSigningCertAndProvisionTutorial");
+		bReadyToBuild |= ETargetPlatformReadyStatus::ProvisionNotFound;
+	}
+
+	return bReadyToBuild;
+}
+
+
+/* FIOSTargetPlatform implementation
+ *****************************************************************************/
+
+void FIOSTargetPlatform::PingNetworkDevices()
+{
+	// disabled for now because we find IOS devices from the USB, this is a relic from ULD, but it may be needed in the future
+/*	if (!MessageEndpoint.IsValid())
+	{
+		MessageEndpoint = FMessageEndpoint::Builder("FIOSTargetPlatform")
+			.Handling<FIOSLaunchDaemonPong>(this, &FIOSTargetPlatform::HandlePongMessage);
+	}
+
+	if (MessageEndpoint.IsValid())
+	{
+		MessageEndpoint->Publish(new FIOSLaunchDaemonPing(), EMessageScope::Network);
+	}
+
+	// remove disconnected & timed out devices
+	FDateTime Now = FDateTime::UtcNow();
+
+	for (auto DeviceIt = Devices.CreateIterator(); DeviceIt; ++DeviceIt)
+	{
+		FIOSTargetDevicePtr Device = DeviceIt->Value;
+
+		if (Now > Device->LastPinged + FTimespan::FromSeconds(60.0))
+		{
+			DeviceIt.RemoveCurrent();
+			DeviceLostEvent.Broadcast(Device.ToSharedRef());
+		}
+	}*/
+}
+
+
+/* FIOSTargetPlatform callbacks
+ *****************************************************************************/
+
+void FIOSTargetPlatform::HandlePongMessage( const FIOSLaunchDaemonPong& Message, const IMessageContextRef& Context )
+{
+	FTargetDeviceId DeviceId;
+	FTargetDeviceId::Parse(Message.DeviceID, DeviceId);
+
+	FIOSTargetDevicePtr& Device = Devices.FindOrAdd(DeviceId);
+
+	if (!Device.IsValid())
+	{
+		Device = MakeShareable(new FIOSTargetDevice(*this));
+
+		Device->SetFeature(ETargetDeviceFeatures::Reboot, Message.bCanReboot);
+		Device->SetFeature(ETargetDeviceFeatures::PowerOn, Message.bCanPowerOn);
+		Device->SetFeature(ETargetDeviceFeatures::PowerOff, Message.bCanPowerOff);
+		Device->SetDeviceId(DeviceId);
+		Device->SetDeviceName(Message.DeviceName);
+		Device->SetDeviceType(Message.DeviceType);
+		Device->SetDeviceEndpoint(Context->GetSender());
+		Device->SetIsSimulated(Message.DeviceID.Contains(TEXT("Simulator")));
+
+		DeviceDiscoveredEvent.Broadcast(Device.ToSharedRef());
+	}
+
+	Device->LastPinged = FDateTime::UtcNow();
+}
+
+void FIOSTargetPlatform::HandleDeviceConnected(const FIOSLaunchDaemonPong& Message)
+{
+	FTargetDeviceId DeviceId;
+	FTargetDeviceId::Parse(Message.DeviceID, DeviceId);
+	
+	FIOSTargetDevicePtr& Device = Devices.FindOrAdd(DeviceId);
+	
+	if (!Device.IsValid())
+	{
+		if ((Message.DeviceType.Contains(TEXT("AppleTV")) && bIsTVOS) || (!Message.DeviceType.Contains(TEXT("AppleTV")) && !bIsTVOS))
+		{
+			Device = MakeShareable(new FIOSTargetDevice(*this));
+
+			Device->SetFeature(ETargetDeviceFeatures::Reboot, Message.bCanReboot);
+			Device->SetFeature(ETargetDeviceFeatures::PowerOn, Message.bCanPowerOn);
+			Device->SetFeature(ETargetDeviceFeatures::PowerOff, Message.bCanPowerOff);
+			Device->SetDeviceId(DeviceId);
+			Device->SetDeviceName(Message.DeviceName);
+			Device->SetDeviceType(Message.DeviceType);
+			Device->SetIsSimulated(Message.DeviceID.Contains(TEXT("Simulator")));
+
+			DeviceDiscoveredEvent.Broadcast(Device.ToSharedRef());
+		}
+		else
+		{
+			return;
+		}
+	}
+	
+	// Add a very long time period to prevent the devices from getting disconnected due to a lack of pong messages
+	Device->LastPinged = FDateTime::UtcNow() + FTimespan(100, 0, 0, 0, 0);
+}
+
+
+void FIOSTargetPlatform::HandleDeviceDisconnected(const FIOSLaunchDaemonPong& Message)
+{
+	FTargetDeviceId DeviceId;
+	FTargetDeviceId::Parse(Message.DeviceID, DeviceId);
+	
+	FIOSTargetDevicePtr& Device = Devices.FindOrAdd(DeviceId);
+	
+	if (Device.IsValid())
+	{
+		DeviceLostEvent.Broadcast(Device.ToSharedRef());
+		Devices.Remove(DeviceId);
+	}
+}
+
+bool FIOSTargetPlatform::HandleTicker(float DeltaTime )
+{
+	PingNetworkDevices();
+
+	return true;
+}
+
+
+/* ITargetPlatform interface
+ *****************************************************************************/
+
+static bool SupportsES2()
+{
+	// default to supporting ES2
+	bool bSupportsOpenGLES2 = true;
+	GConfig->GetBool(TEXT("/Script/IOSRuntimeSettings.IOSRuntimeSettings"), TEXT("bSupportsOpenGLES2"), bSupportsOpenGLES2, GEngineIni);
+	return bSupportsOpenGLES2;
+}
+
+static bool SupportsMetal()
+{
+	// default to NOT supporting metal
+	bool bSupportsMetal = false;
+	GConfig->GetBool(TEXT("/Script/IOSRuntimeSettings.IOSRuntimeSettings"), TEXT("bSupportsMetal"), bSupportsMetal, GEngineIni);
+	return bSupportsMetal;
+}
+
+static bool SupportsMetalMRT()
+{
+	// default to NOT supporting metal MRT
+	bool bSupportsMetalMRT = false;
+	GConfig->GetBool(TEXT("/Script/IOSRuntimeSettings.IOSRuntimeSettings"), TEXT("bSupportsMetalMRT"), bSupportsMetalMRT, GEngineIni);
+	return bSupportsMetalMRT;
+}
+
+static bool CookPVRTC()
+{
+	// default to using PVRTC
+	bool bCookPVRTCTextures = true;
+	GConfig->GetBool(TEXT("/Script/IOSRuntimeSettings.IOSRuntimeSettings"), TEXT("bCookPVRTCTextures"), bCookPVRTCTextures, GEngineIni);
+	return bCookPVRTCTextures;
+}
+
+static bool CookASTC()
+{
+	// default to not using ASTC
+	bool bCookASTCTextures = true;
+	GConfig->GetBool(TEXT("/Script/IOSRuntimeSettings.IOSRuntimeSettings"), TEXT("bCookASTCTextures"), bCookASTCTextures, GEngineIni);
+	return bCookASTCTextures;
+}
+
+bool FIOSTargetPlatform::SupportsFeature( ETargetPlatformFeatures Feature ) const
+{
+	switch (Feature)
+	{
+		case ETargetPlatformFeatures::Packaging:
+			return true;
+
+		case ETargetPlatformFeatures::MobileRendering:
+		case ETargetPlatformFeatures::LowQualityLightmaps:
+			return SupportsES2() || SupportsMetal();
+			
+		case ETargetPlatformFeatures::DeferredRendering:
+		case ETargetPlatformFeatures::HighQualityLightmaps:
+			return SupportsMetalMRT();
+
+		default:
+			break;
+	}
+	
+	return TTargetPlatformBase<FIOSPlatformProperties>::SupportsFeature(Feature);
+}
+
+
+#if WITH_ENGINE
+
+
+void FIOSTargetPlatform::GetAllPossibleShaderFormats( TArray<FName>& OutFormats ) const
+{
+	static FName NAME_GLSL_ES2_IOS(TEXT("GLSL_ES2_IOS"));
+	static FName NAME_SF_METAL(TEXT("SF_METAL"));
+	static FName NAME_SF_METAL_MRT(TEXT("SF_METAL_MRT"));
+
+	if (bIsTVOS)
+	{
+		if (SupportsMetalMRT())
+		{
+			OutFormats.AddUnique(NAME_SF_METAL_MRT);
+		}
+
+		// because we are currently using IOS settings, we will always use metal, even if Metal isn't listed as being supported
+		// however, if MetalMRT is specific and Metal is set to false, then we will just use MetalMRT
+		if (SupportsMetal() || !SupportsMetalMRT())
+		{
+			OutFormats.AddUnique(NAME_SF_METAL);
+		}
+	}
+	else
+	{
+		if (SupportsES2())
+		{
+			OutFormats.AddUnique(NAME_GLSL_ES2_IOS);
+		}
+
+		if (SupportsMetal())
+		{
+			OutFormats.AddUnique(NAME_SF_METAL);
+		}
+
+		if (SupportsMetalMRT())
+		{
+			OutFormats.AddUnique(NAME_SF_METAL_MRT);
+		}
+	}
+}
+
+void FIOSTargetPlatform::GetAllTargetedShaderFormats( TArray<FName>& OutFormats ) const
+{
+	GetAllPossibleShaderFormats(OutFormats);
+}
+
+// we remap some of the defaults (with PVRTC and ASTC formats)
+static FName FormatRemap[] =
+{
+	// original				PVRTC						ASTC
+	FName(TEXT("DXT1")),	FName(TEXT("PVRTC2")),		FName(TEXT("ASTC_RGB")),
+	FName(TEXT("DXT5")),	FName(TEXT("PVRTC4")),		FName(TEXT("ASTC_RGBA")),
+	FName(TEXT("DXT5n")),	FName(TEXT("PVRTCN")),		FName(TEXT("ASTC_NormalAG")),
+	FName(TEXT("BC5")),		FName(TEXT("PVRTCN")),		FName(TEXT("ASTC_NormalRG")),
+	FName(TEXT("AutoDXT")),	FName(TEXT("AutoPVRTC")),	FName(TEXT("ASTC_RGBAuto")),
+	FName(TEXT("BC4")),		FName(TEXT("G8")),			FName(TEXT("G8")),
+};
+static FName NameBGRA8(TEXT("BGRA8"));
+static FName NameG8 = FName(TEXT("G8"));
+
+void FIOSTargetPlatform::GetTextureFormats( const UTexture* Texture, TArray<FName>& OutFormats ) const
+{
+	check(Texture);
+
+	static FName NamePOTERROR(TEXT("POTERROR"));
+
+	FName TextureFormatName = NAME_None;
+
+	// forward rendering only needs one channel for shadow maps
+	if (Texture->LODGroup == TEXTUREGROUP_Shadowmap && !SupportsMetalMRT())
+	{
+		TextureFormatName = NameG8;
+	}
+
+	// if we didn't assign anything specially, then use the defaults
+	if (TextureFormatName == NAME_None)
+	{
+		TextureFormatName = GetDefaultTextureFormatName(this, Texture, EngineSettings, false);
+	}
+
+	// perform any remapping away from defaults
+	bool bFoundRemap = false;
+	bool bIncludePVRTC = !bIsTVOS && CookPVRTC();
+	bool bIncludeASTC = bIsTVOS || CookASTC();
+	for (int32 RemapIndex = 0; RemapIndex < ARRAY_COUNT(FormatRemap); RemapIndex += 3)
+	{
+		if (TextureFormatName == FormatRemap[RemapIndex])
+		{
+			// we found a remapping
+			bFoundRemap = true;
+			// include the formats we want (use ASTC first so that it is preferred at runtime if they both exist and it's supported)
+			if (bIncludeASTC)
+			{
+				OutFormats.AddUnique(FormatRemap[RemapIndex + 2]);
+			}
+			if (bIncludePVRTC)
+			{
+				// handle non-power of 2 textures
+				if (!Texture->Source.IsPowerOfTwo() && Texture->PowerOfTwoMode == ETexturePowerOfTwoSetting::None)
+				{
+					// option 1: Uncompress, but users will get very large textures unknowningly
+					// OutFormats.AddUnique(NameBGRA8);
+					// option 2: Use an "error message" texture so they see it in game
+					OutFormats.AddUnique(NamePOTERROR);
+				}
+				else
+				{
+					OutFormats.AddUnique(FormatRemap[RemapIndex + 1]);
+				}
+			}
+		}
+	}
+
+	// if we didn't already remap above, add it now
+	if (!bFoundRemap)
+	{
+		OutFormats.Add(TextureFormatName);
+	}
+}
+
+void FIOSTargetPlatform::GetAllTextureFormats(TArray<FName>& OutFormats) const 
+{
+	bool bFoundRemap = false;
+	bool bIncludePVRTC = !bIsTVOS && CookPVRTC();
+	bool bIncludeASTC = bIsTVOS || CookASTC();
+
+	GetAllDefaultTextureFormats(this, OutFormats, false);
+
+	for (int32 RemapIndex = 0; RemapIndex < ARRAY_COUNT(FormatRemap); RemapIndex += 3)
+	{
+		OutFormats.Remove(FormatRemap[RemapIndex+0]);
+	}
+
+	// include the formats we want (use ASTC first so that it is preferred at runtime if they both exist and it's supported)
+	if (bIncludeASTC)
+	{
+		for (int32 RemapIndex = 0; RemapIndex < ARRAY_COUNT(FormatRemap); RemapIndex += 3)
+		{
+			OutFormats.AddUnique(FormatRemap[RemapIndex + 2]);
+		}
+	}
+	if (bIncludePVRTC)
+	{
+		for (int32 RemapIndex = 0; RemapIndex < ARRAY_COUNT(FormatRemap); RemapIndex += 3)
+		{
+			OutFormats.AddUnique(FormatRemap[RemapIndex + 1]);
+		}
+	}
+}
+
+
+const UTextureLODSettings& FIOSTargetPlatform::GetTextureLODSettings() const
+{
+	return *TextureLODSettings;
+}
+
+
+FName FIOSTargetPlatform::GetWaveFormat( const class USoundWave* Wave ) const
+{
+	static FName NAME_ADPCM(TEXT("ADPCM"));
+	return NAME_ADPCM;
+}
+
+
+void FIOSTargetPlatform::GetAllWaveFormats(TArray<FName>& OutFormat) const
+{
+	static FName NAME_ADPCM(TEXT("ADPCM"));
+	OutFormat.Add(NAME_ADPCM);
+}
+
+#endif // WITH_ENGINE
+
diff --git a/Engine/Source/Developer/IOS/IOSTargetPlatform/Private/IOSTargetPlatform.h b/Engine/Source/Developer/IOS/IOSTargetPlatform/Private/IOSTargetPlatform.h
new file mode 100644
--- /dev/null
+++ b/Engine/Source/Developer/IOS/IOSTargetPlatform/Private/IOSTargetPlatform.h
@@ -0,0 +1,219 @@
+// Copyright 1998-2017 Epic Games, Inc. All Rights Reserved.
+
+/*=============================================================================
+	IOSTargetPlatform.h: Declares the FIOSTargetPlatform class.
+=============================================================================*/
+
+#pragma once
+
+#include "CoreMinimal.h"
+#include "TargetPlatformBase.h"
+#include "IOSPlatformProperties.h"
+#include "Ticker.h"
+#include "IOSMessageProtocol.h"
+#include "IMessageContext.h"
+#include "IOSTargetDevice.h"
+#include "IOSDeviceHelper.h"
+#include "Misc/ConfigCacheIni.h"
+
+#if WITH_ENGINE
+#include "StaticMeshResources.h"
+#endif // WITH_ENGINE
+
+/**
+ * FIOSTargetPlatform, abstraction for cooking iOS platforms
+ */
+class FIOSTargetPlatform : public TTargetPlatformBase<FIOSPlatformProperties>
+{
+public:
+
+	/**
+	 * Default constructor.
+	 */
+	IOSTARGETPLATFORM_API FIOSTargetPlatform(bool bInISTVOS);
+
+	/**
+	 * Destructor.
+	 */
+	~FIOSTargetPlatform();
+
+public:
+
+	//~ Begin TTargetPlatformBase Interface
+
+	virtual bool IsServerOnly( ) const override
+	{
+		return false;
+	}
+
+	//~ End TTargetPlatformBase Interface
+
+public:
+
+	//~ Begin ITargetPlatform Interface
+	
+	// this is used for cooking to a separate directory, NOT for runtime. Runtime TVOS is still "IOS"
+	virtual FString PlatformName() const override
+	{
+		return bIsTVOS ? "TVOS" : "IOS";
+	}
+
+    virtual FString IniPlatformName() const override
+    {
+        return "IOS";
+    }
+    
+	virtual void EnableDeviceCheck(bool OnOff) override;
+
+	virtual void GetAllDevices( TArray<ITargetDevicePtr>& OutDevices ) const override;
+
+	virtual ECompressionFlags GetBaseCompressionMethod() const override
+	{
+		return COMPRESS_ZLIB;
+	}
+
+	virtual bool GenerateStreamingInstallManifest(const TMultiMap<FString, int32>& ChunkMap, const TSet<int32>& ChunkIDsInUse) const override
+	{
+		return true;
+	}
+
+	virtual ITargetDevicePtr GetDefaultDevice( ) const override;
+
+	virtual ITargetDevicePtr GetDevice( const FTargetDeviceId& DeviceId ) override;
+
+	virtual bool IsRunningPlatform( ) const override
+	{
+		#if PLATFORM_IOS && WITH_EDITOR
+			return true;
+		#else
+			return false;
+		#endif
+	}
+
+	virtual bool SupportsFeature( ETargetPlatformFeatures Feature ) const override;
+
+	virtual bool IsSdkInstalled(bool bProjectHasCode, FString& OutTutorialPath) const override;
+	virtual int32 CheckRequirements(const FString& ProjectPath, bool bProjectHasCode, FString& OutTutorialPath) const override;
+
+
+#if WITH_ENGINE
+	virtual void GetReflectionCaptureFormats( TArray<FName>& OutFormats ) const override
+	{
+		OutFormats.Add(FName(TEXT("EncodedHDR")));
+		OutFormats.Add(FName(TEXT("FullHDR")));
+	}
+
+	virtual void GetAllPossibleShaderFormats( TArray<FName>& OutFormats ) const override;
+
+	virtual void GetAllTargetedShaderFormats( TArray<FName>& OutFormats ) const override;
+	virtual const class FStaticMeshLODSettings& GetStaticMeshLODSettings( ) const override
+	{
+		return StaticMeshLODSettings;
+	}
+
+	virtual void GetTextureFormats( const UTexture* Texture, TArray<FName>& OutFormats ) const override;
+
+	virtual void GetAllTextureFormats( TArray<FName>& OutFormats) const override;
+
+	virtual const UTextureLODSettings& GetTextureLODSettings() const override;
+
+	virtual void RegisterTextureLODSettings(const UTextureLODSettings* InTextureLODSettings) override
+	{
+		TextureLODSettings = InTextureLODSettings;
+	}
+
+	virtual FName GetWaveFormat( const class USoundWave* Wave ) const override;
+	virtual void GetAllWaveFormats(TArray<FName>& OutFormat) const override;
+#endif // WITH_ENGINE
+
+	virtual void GetBuildProjectSettingKeys(FString& OutSection, TArray<FString>& InBoolKeys, TArray<FString>& InIntKeys, TArray<FString>& InStringKeys) const override
+	{
+		OutSection = TEXT("/Script/IOSRuntimeSettings.IOSRuntimeSettings");
+		InBoolKeys.Add(TEXT("EnableRemoteShaderCompile"));
+		InBoolKeys.Add(TEXT("bGeneratedSYMFile"));
+		InBoolKeys.Add(TEXT("bGeneratedSYMBundle"));
+		InBoolKeys.Add(TEXT("bGenerateXCArchive"));
+		InBoolKeys.Add(TEXT("bShipForBitcode"));
+		if (bIsTVOS)
+		{
+			InStringKeys.Add(TEXT("MinimumTVOSVersion"));
+		}
+		else
+		{
+			InStringKeys.Add(TEXT("MinimumiOSVersion"));
+			InBoolKeys.Add(TEXT("bDevForArmV7")); InBoolKeys.Add(TEXT("bDevForArm64")); InBoolKeys.Add(TEXT("bDevForArmV7S"));
+			InBoolKeys.Add(TEXT("bShipForArmV7")); InBoolKeys.Add(TEXT("bShipForArm64")); InBoolKeys.Add(TEXT("bShipForArmV7S"));
+		}
+	}
+
+	DECLARE_DERIVED_EVENT(FIOSTargetPlatform, ITargetPlatform::FOnTargetDeviceDiscovered, FOnTargetDeviceDiscovered);
+	virtual FOnTargetDeviceDiscovered& OnDeviceDiscovered( ) override
+	{
+		return DeviceDiscoveredEvent;
+	}
+
+	DECLARE_DERIVED_EVENT(FIOSTargetPlatform, ITargetPlatform::FOnTargetDeviceLost, FOnTargetDeviceLost);
+	virtual FOnTargetDeviceLost& OnDeviceLost( ) override
+	{
+		return DeviceLostEvent;
+	}
+
+	//~ Begin ITargetPlatform Interface
+
+protected:
+
+	/**
+	 * Sends a ping message over the network to find devices running the launch daemon.
+	 */
+	void PingNetworkDevices( );
+
+private:
+
+	// Handles when the ticker fires.
+	bool HandleTicker( float DeltaTime );
+
+	// Handles received pong messages from the LauncherDaemon.
+	void HandlePongMessage( const FIOSLaunchDaemonPong& Message, const IMessageContextRef& Context );
+
+    void HandleDeviceConnected( const FIOSLaunchDaemonPong& Message );
+    void HandleDeviceDisconnected( const FIOSLaunchDaemonPong& Message );
+
+private:
+	
+	// true if this is targeting TVOS vs IOS
+	bool bIsTVOS;
+
+	// Contains all discovered IOSTargetDevices over the network.
+	TMap<FTargetDeviceId, FIOSTargetDevicePtr> Devices;
+
+	// Holds a delegate to be invoked when the widget ticks.
+	FTickerDelegate TickDelegate;
+
+	// Handle to the registered TickDelegate.
+	FDelegateHandle TickDelegateHandle;
+
+	// Holds the message endpoint used for communicating with the LaunchDaemon.
+	FMessageEndpointPtr MessageEndpoint;
+
+#if WITH_ENGINE
+	// Holds the Engine INI settings, for quick use.
+	FConfigFile EngineSettings;
+
+	// Holds the cache of the target LOD settings.
+	const UTextureLODSettings* TextureLODSettings;
+
+	// Holds the static mesh LOD settings.
+	FStaticMeshLODSettings StaticMeshLODSettings;
+#endif // WITH_ENGINE
+
+    // holds usb device helper
+	FIOSDeviceHelper DeviceHelper;
+
+private:
+
+	// Holds an event delegate that is executed when a new target device has been discovered.
+	FOnTargetDeviceDiscovered DeviceDiscoveredEvent;
+
+	// Holds an event delegate that is executed when a target device has been lost, i.e. disconnected or timed out.
+	FOnTargetDeviceLost DeviceLostEvent;
+};
diff --git a/Engine/Source/Developer/IOS/IOSTargetPlatform/Private/IOSTargetPlatformModule.cpp b/Engine/Source/Developer/IOS/IOSTargetPlatform/Private/IOSTargetPlatformModule.cpp
new file mode 100644
--- /dev/null
+++ b/Engine/Source/Developer/IOS/IOSTargetPlatform/Private/IOSTargetPlatformModule.cpp
@@ -0,0 +1,58 @@
+// Copyright 1998-2017 Epic Games, Inc. All Rights Reserved.
+
+#include "IOSTargetPlatform.h"
+#include "Interfaces/ITargetPlatformModule.h"
+
+/**
+ * Holds the target platform singleton.
+ */
+static ITargetPlatform* Singleton = NULL;
+
+
+/**
+ * Module for iOS as a target platform
+ */
+class FIOSTargetPlatformModule : public ITargetPlatformModule
+{
+public:
+
+	/**
+	 * Destructor.
+	 */
+	~FIOSTargetPlatformModule()
+	{
+		Singleton = NULL;
+	}
+
+
+public:
+
+	// Begin ITargetPlatformModule interface
+
+	virtual ITargetPlatform* GetTargetPlatform()
+	{
+		if (Singleton == NULL)
+		{
+			Singleton = new FIOSTargetPlatform(false);
+		}
+
+		return Singleton;
+	}
+
+	// End ITargetPlatformModule interface
+
+
+public:
+
+	// Begin IModuleInterface interface
+	virtual void StartupModule() override
+	{
+	}
+
+	virtual void ShutdownModule() override
+	{
+	}
+	// End IModuleInterface interface
+};
+
+IMPLEMENT_MODULE( FIOSTargetPlatformModule, IOSTargetPlatform);
diff --git a/Engine/Source/Developer/IOS/IOSTargetPlatform/Private/Mac/IOSDeviceHelperMac.cpp b/Engine/Source/Developer/IOS/IOSTargetPlatform/Private/Mac/IOSDeviceHelperMac.cpp
new file mode 100644
--- /dev/null
+++ b/Engine/Source/Developer/IOS/IOSTargetPlatform/Private/Mac/IOSDeviceHelperMac.cpp
@@ -0,0 +1,397 @@
+// Copyright 1998-2017 Epic Games, Inc. All Rights Reserved
+
+#include "IOSTargetPlatform.h"
+
+#include <CoreFoundation/CoreFoundation.h>
+
+struct AMDeviceNotificationCallbackInformation
+{
+	void 		*deviceHandle;
+	uint32_t	msgType;
+};
+
+struct AFCCommConnection
+{
+};
+
+#ifdef __cplusplus
+extern "C" {
+#endif
+
+int AMDeviceConnect (void *device);
+int AMDeviceValidatePairing (void *device);
+int AMDeviceStartSession (void *device);
+int AMDeviceStopSession (void *device);
+int AMDeviceDisconnect (void *device);
+int AMDeviceNotificationSubscribe(void*, int, int, int, void**);
+int AMDeviceStartService(void*, CFStringRef, void**, int);
+CFStringRef AMDeviceCopyValue(void*, int unknown, CFStringRef cfstring);
+int AMDeviceSecureUpgradeApplication(void*, void*, CFURLRef, void*, void*, void*);
+    
+uint32 AFCConnectionOpen(void*, uint32, void**);
+uint32 AFCConnectionClose(void*);
+uint32 AFCDirectoryCreate(void*, const char*);
+uint32 AFCFileRefOpen(void*, const char*, uint64, uint64*);
+uint32 AFCFileRefClose(void*, uint64);
+uint32 AFCFileRefWrite(void*, uint64, void*, uint32);
+#ifdef __cplusplus
+}
+#endif
+
+class AFC
+{
+public:
+    static int StartService(void* DeviceHandle, void** OutHandle)
+    {
+        int Result = AMDeviceStartService(DeviceHandle, CFSTR("com.apple.afc"), OutHandle, 0);
+        return Result;
+    }
+    
+    static int ConnectionOpen(void* ServiceHandle, AFCCommConnection** OutConn)
+    {
+        void* Conn = NULL;
+        int result = AFCConnectionOpen(ServiceHandle, 0, &Conn);
+        *OutConn = (AFCCommConnection*)Conn;
+        return result;
+    }
+    
+    static int ConnectionClose(AFCCommConnection* Connection)
+    {
+        return AFCConnectionClose(Connection);
+    }
+    
+    static uint32 DirectoryCreate(AFCCommConnection* Connection, FString Dir)
+    {
+        const char* IPAPath = TCHAR_TO_UTF8(*Dir);
+        uint Result = AFCDirectoryCreate(Connection, IPAPath);
+        return Result;
+    }
+    
+    static uint32 FileRefOpen(AFCCommConnection* Connection, FString Path, int mode, uint64* OutHandle)
+    {
+        const char* IPAPath = TCHAR_TO_UTF8(*Path);
+        uint Result = AFCFileRefOpen(Connection, IPAPath, mode, OutHandle);
+        return Result;
+    }
+    
+    static uint32 FileRefClose(AFCCommConnection* Connection, uint64 Handle)
+    {
+        return AFCFileRefClose(Connection, Handle);
+    }
+    
+    static uint32 FileRefWrite(AFCCommConnection* Connection, uint64 Handle, void* Data, uint32 Length)
+    {
+        return AFCFileRefWrite(Connection, Handle, Data, Length);
+    }
+};
+
+class InstallProxy
+{
+public:
+    static int StartService(void* DeviceHandle, void** OutHandle)
+    {
+        int Result = AMDeviceStartService(DeviceHandle, CFSTR("com.apple.mobile.installation_proxy"), OutHandle, 0);
+        return Result;
+    }
+    
+    static int SecureUpgradeApplication(void* ServiceConnection, void* DeviceHandle, CFURLRef UrlPath, void* Options, void* InstallCallback, void* UserData)
+    {
+        return AMDeviceSecureUpgradeApplication(ServiceConnection, DeviceHandle, UrlPath, Options, InstallCallback, UserData);
+    }
+};
+
+class IOSDevice
+{
+public:
+    IOSDevice(void* device)
+    : DeviceHandle(device)
+    , AFCHandle(NULL)
+    , InstallHandle(NULL)
+    {
+    }
+    
+    bool Connect()
+    {
+        bool Result = false;
+		int AppLogLevel = 7;
+		CFPreferencesSetAppValue(CFSTR("LogLevel"), CFNumberCreate(kCFAllocatorDefault, kCFNumberSInt32Type, &AppLogLevel), CFSTR("com.apple.MobileDevice"));
+
+		// connect to the device
+        int32 rc = AMDeviceConnect(DeviceHandle);
+        if (!rc)
+        {
+            // validate the pairing
+            rc = AMDeviceValidatePairing(DeviceHandle);
+            if (!rc)
+            {
+                // start a session
+                rc = AMDeviceStartSession(DeviceHandle);
+                Result = !rc;
+            }
+            else
+            {
+                UE_LOG(LogTemp, Display, TEXT("Couldn't validate pairing to device"));
+            }
+        }
+        else
+        {
+            UE_LOG(LogTemp, Display, TEXT("Couldn't connect to device"));
+        }
+        return Result;
+    }
+    
+    bool Disconnect()
+    {
+        bool Result = false;
+        // close the session
+        int32 rc = AMDeviceStopSession(DeviceHandle);
+        if (!rc)
+        {
+            // disconnect from the device
+            rc = AMDeviceDisconnect(DeviceHandle);
+            Result = !rc;
+        }
+        else
+        {
+            UE_LOG(LogTemp, Display, TEXT("Couldn't stop session"));
+        }
+        return Result;
+    }
+    
+    bool CopyFileToPublicStaging(const FString& SourceFile)
+    {
+        FString IpaFilename = FPaths::GetCleanFilename(SourceFile);
+        return CopyFileToDevice(SourceFile, FString(TEXT("/PublicStaging/")) + IpaFilename, 1024*1024);
+    }
+    
+    bool TryUpgrade(const FString& IPAPath)
+    {
+        // reconnect to the device and start the Installation service
+        Connect();
+        InstallProxy::StartService(DeviceHandle, &InstallHandle);
+
+        CFURLRef UrlPath = CFURLCreateWithFileSystemPath(kCFAllocatorDefault, FPlatformString::TCHARToCFString(*IPAPath), kCFURLWindowsPathStyle, 0);
+        int Result = InstallProxy::SecureUpgradeApplication(0, DeviceHandle, UrlPath, 0, 0, 0);
+        return Result == 0;
+    }
+    
+    void* Handle() const
+    {
+        return DeviceHandle;
+    }
+    
+    bool CreateDirectory(FString Dir)
+    {
+        return !(AFC::DirectoryCreate(AFCConnection, Dir) != 0);
+    }
+    
+private:
+    
+    bool CopyFileToDevice(const FString& IPAPath, FString PathOnDevice, int PacketSize)
+    {
+        bool Result = false;
+        // reconnect to the device and start the AFC service
+        Connect();
+        if (!AFC::StartService(DeviceHandle, &AFCHandle))
+        {
+            if (!AFC::ConnectionOpen(AFCHandle, &AFCConnection))
+            {
+        
+                // ensure the directory on the phone exists
+                FString DirectoryOnDevice = FPaths::GetPath(PathOnDevice);
+                CreateDirectory(DirectoryOnDevice + TEXT("/"));
+        
+                // transfer the file
+                IPlatformFile& PlatformFile = IPlatformFile::GetPlatformPhysical();
+                IFileHandle* SourceFile = PlatformFile.OpenRead(*IPAPath);
+                if (SourceFile != NULL)
+                {
+                    uint64 DestinationHandle = 0;
+                    if (AFC::FileRefOpen(AFCConnection, PathOnDevice, 3, &DestinationHandle) == 0)
+                    {
+                        int TotalBytes = 0;
+                        int PacketCount = SourceFile->Size() / PacketSize;
+                        uint8* buffer = new uint8[PacketSize];
+                        for (int Index = 0; Index < PacketCount; ++Index)
+                        {
+                            if (SourceFile->Read(buffer, PacketSize))
+                            {
+                                TotalBytes += PacketSize;
+                                AFC::FileRefWrite(AFCConnection, DestinationHandle, buffer, PacketSize);
+                            }
+                        }
+                
+                        if (SourceFile->Read(buffer, SourceFile->Size() - TotalBytes))
+                        {
+                            AFC::FileRefWrite(AFCConnection, DestinationHandle, buffer, SourceFile->Size() - TotalBytes);
+                        }
+                
+                        // flush the destination and close
+                        AFC::FileRefClose(AFCConnection, DestinationHandle);
+                        
+                        Result = true;
+                    }
+                    delete SourceFile;
+                }
+        
+                // stop the AFC service and disconnect from the device
+                AFC::ConnectionClose(AFCConnection);
+                AFCConnection = NULL;
+                Disconnect();
+            }
+        }
+        return Result;
+    }
+    
+	void* DeviceHandle;
+    void* AFCHandle;
+    AFCCommConnection* AFCConnection;
+    void* InstallHandle;
+};
+
+
+/* FIOSDeviceHelper structors
+ *****************************************************************************/
+
+static TMap<IOSDevice*, FIOSLaunchDaemonPong> ConnectedDevices;
+
+void FIOSDeviceHelper::Initialize(bool bIsTVOS)
+{
+    static bool bIsInitialized = false;
+    
+    if (!bIsInitialized)
+    {
+        void *subscribe;
+        int32 rc = AMDeviceNotificationSubscribe((void*)FIOSDeviceHelper::DeviceCallback, 0, 0, 0, &subscribe);
+        if (rc < 0)
+        {
+            //@todo right out to the log that we can't subscribe
+        }
+        bIsInitialized = true;
+    }
+}
+
+void FIOSDeviceHelper::DeviceCallback(void* CallbackInfo)
+{
+    struct AMDeviceNotificationCallbackInformation* cbi = (AMDeviceNotificationCallbackInformation*)CallbackInfo;
+    
+    void* deviceHandle = cbi->deviceHandle;
+    
+    switch(cbi->msgType)
+    {
+        case 1:
+            FIOSDeviceHelper::DoDeviceConnect(deviceHandle);
+            break;
+            
+        case 2:
+            FIOSDeviceHelper::DoDeviceDisconnect(deviceHandle);
+            break;
+            
+        case 3:
+            break;
+    }
+}
+
+void FIOSDeviceHelper::DoDeviceConnect(void* deviceHandle)
+{
+    // connect to the device
+    IOSDevice* Device = new IOSDevice(deviceHandle);
+    bool Connected = Device->Connect();
+    if (Connected)
+    {
+        // get the needed data
+        CFStringRef deviceName = AMDeviceCopyValue(deviceHandle, 0, CFSTR("DeviceName"));
+        CFStringRef deviceId = AMDeviceCopyValue(deviceHandle, 0, CFSTR("UniqueDeviceID"));
+		CFStringRef productType = AMDeviceCopyValue(deviceHandle, 0, CFSTR("ProductType"));
+
+        // fire the event
+        TCHAR idBuffer[128];
+		TCHAR nameBuffer[256];
+		TCHAR productBuffer[128];
+		FPlatformString::CFStringToTCHAR(deviceId, idBuffer);
+		FPlatformString::CFStringToTCHAR(deviceName, nameBuffer);
+		FPlatformString::CFStringToTCHAR(productType, productBuffer);
+		FIOSLaunchDaemonPong Event;
+        Event.DeviceID = FString::Printf(TEXT("%s@%s"), (FString(productBuffer).Contains("AppleTV") ? TEXT("TVOS") : TEXT("IOS")), idBuffer);
+		Event.DeviceName = FString::Printf(TEXT("%s"), nameBuffer);
+        Event.bCanReboot = false;
+        Event.bCanPowerOn = false;
+        Event.bCanPowerOff = false;
+		Event.DeviceType = FString::Printf(TEXT("%s"), productBuffer);
+		FIOSDeviceHelper::OnDeviceConnected().Broadcast(Event);
+                
+        // add to the device list
+        ConnectedDevices.Add(Device, Event);
+        
+        // disconnect the device for now
+        Device->Disconnect();
+    }
+    else
+    {
+        UE_LOG(LogTemp, Display, TEXT("Couldn't connect to device"));
+    }
+}
+
+void FIOSDeviceHelper::DoDeviceDisconnect(void* deviceHandle)
+{
+    IOSDevice* device = NULL;
+    for (auto DeviceIterator = ConnectedDevices.CreateIterator(); DeviceIterator; ++DeviceIterator)
+    {
+        if (DeviceIterator.Key()->Handle() == deviceHandle)
+        {
+            device = DeviceIterator.Key();
+            break;
+        }
+    }
+    if (device != NULL)
+    {
+        // extract the device id from the connected list
+		FIOSLaunchDaemonPong Event = ConnectedDevices.FindAndRemoveChecked(device);
+    
+        // fire the event
+        FIOSDeviceHelper::OnDeviceDisconnected().Broadcast(Event);
+        
+        // delete the device
+        delete device;
+    }
+}
+
+bool FIOSDeviceHelper::InstallIPAOnDevice(const FTargetDeviceId& DeviceId, const FString& IPAPath)
+{
+    // check for valid path
+    if (IPAPath.Len() == 0)
+    {
+        return false;
+    }
+    
+    // check for valid device
+    IOSDevice* device = NULL;
+    FIOSLaunchDaemonPong DeviceMessage;
+    for (auto DeviceIterator = ConnectedDevices.CreateIterator(); DeviceIterator; ++DeviceIterator)
+    {
+        DeviceMessage = DeviceIterator.Value();
+        if (DeviceMessage.DeviceID == DeviceId.ToString())
+        {
+            device = DeviceIterator.Key();
+            break;
+        }
+    }
+    if (device == NULL)
+    {
+        return false;
+    }
+    
+    // we have the device and a IPA path
+    // copy to the stage
+    if (device->CopyFileToPublicStaging(IPAPath))
+    {
+        // install on the device
+        return device->TryUpgrade(IPAPath);
+    }
+    return false;
+}
+
+void FIOSDeviceHelper::EnableDeviceCheck(bool OnOff)
+{
+
+}
\ No newline at end of file
diff --git a/Engine/Source/Developer/IOS/IOSTargetPlatform/Private/Windows/IOSDeviceHelperWindows.cpp b/Engine/Source/Developer/IOS/IOSTargetPlatform/Private/Windows/IOSDeviceHelperWindows.cpp
new file mode 100644
--- /dev/null
+++ b/Engine/Source/Developer/IOS/IOSTargetPlatform/Private/Windows/IOSDeviceHelperWindows.cpp
@@ -0,0 +1,452 @@
+// Copyright 1998-2017 Epic Games, Inc. All Rights Reserved
+
+#include "IOSTargetPlatform.h"
+#include "HAL/Runnable.h"
+#include "HAL/RunnableThread.h"
+
+struct FDeviceNotificationCallbackInformation
+{
+	FString UDID;
+	FString	DeviceName;
+	uint32 msgType;
+};
+
+class FDeviceSyslogRelay : public FRunnable
+{
+public:
+
+	FDeviceSyslogRelay(FString const&	inDeviceID)
+	{
+		deviceID = inDeviceID;
+		Stopping = false;
+	}
+
+	virtual bool Init() override
+	{
+		return true;
+	}
+
+	virtual uint32 Run() override
+	{
+		// Tell DeploymentServer.exe to start collecting logs
+		FString ExecutablePath = FString::Printf(TEXT("%sBinaries/DotNET/IOS"), *FPaths::EngineDir());
+		FString Filename = TEXT("DeploymentServer.exe");
+
+		// Construct command line
+		FString CommandLine = FString::Printf(TEXT("listentodevice -device %s"), *deviceID);
+
+		// execute the command
+		void* WritePipe;
+		void* ReadPipe;
+		FPlatformProcess::CreatePipe(ReadPipe, WritePipe);
+		FProcHandle ProcessHandle = FPlatformProcess::CreateProc(*(ExecutablePath / Filename), *CommandLine, false, true, true, NULL, 0, *ExecutablePath, WritePipe);
+		FString	lastPartialLine;
+		while (FPlatformProcess::IsProcRunning(ProcessHandle) && !Stopping)
+		{
+			FString CurData = lastPartialLine + FPlatformProcess::ReadPipe(ReadPipe);
+			if (CurData.Trim().Len() > 0)
+			{
+				// separate out each line
+				TArray<FString> LogLines;
+				CurData = CurData.Replace(TEXT("\r"), TEXT("\n"));
+
+				bool	endsWithNewline = CurData.EndsWith(TEXT("\n"));
+
+				CurData.ParseIntoArray(LogLines, TEXT("\n"), true);
+
+				for (int32 StringIndex = 0; StringIndex < LogLines.Num() - 1; ++StringIndex)
+				{
+					UE_LOG(LogIOS, Log, TEXT("%s"), *LogLines[StringIndex]);
+				}
+
+				if(endsWithNewline)
+				{
+					UE_LOG(LogIOS, Log, TEXT("%s"), *LogLines[LogLines.Num() - 1]);
+					lastPartialLine = "";
+				}
+				else
+				{
+					lastPartialLine = LogLines[LogLines.Num() - 1];
+				}
+			}
+			else
+			{
+				FPlatformProcess::Sleep(0.1);
+			}
+		}
+
+		FPlatformProcess::Sleep(0.25);
+		FString lastData = FPlatformProcess::ReadPipe(ReadPipe);
+		if (lastData.Len() > 0)
+		{
+			UE_LOG(LogIOS, Log, TEXT("%s"), *lastData);
+		}
+
+		FPlatformProcess::ClosePipe(ReadPipe, WritePipe);
+		FPlatformProcess::TerminateProc(ProcessHandle);
+
+		return 0;
+	}
+
+	virtual void Stop() override
+	{
+		Stopping = true;
+	}
+
+	virtual void Exit() override
+	{}
+
+
+private:
+	bool Stopping;
+	FString	deviceID;
+};
+
+class FIOSDevice
+{
+public:
+    FIOSDevice(FString InID, FString InName)
+		: UDID(InID)
+		, Name(InName)
+    {
+		// BHP - Disabling the ios syslog relay because it depends on DeploymentServer running which makes it unwritable 
+		//	which causes problems when packaging because this library gets rebuilt but it can't overwrite it which
+		//	causes errors and other bad behavior - will probably need to move this functionality into its own dll
+		#if 1
+			syslogRelay = nullptr;
+			syslogRelayThread = nullptr;
+		#else
+			syslogRelay = new FDeviceSyslogRelay(InID);
+			syslogRelayThread = FRunnableThread::Create(syslogRelay, TEXT("FIOSDevice.syslogRelay"), 128 * 1024, TPri_Normal);
+		#endif
+    }
+    
+	~FIOSDevice()
+	{
+		if(syslogRelay != nullptr && syslogRelayThread != nullptr)
+		{
+			syslogRelay->Stop();
+			syslogRelayThread->WaitForCompletion();
+			delete syslogRelayThread;
+			delete syslogRelay;
+			syslogRelay = NULL;
+			syslogRelayThread = NULL;
+		}
+	}
+
+	FString SerialNumber() const
+	{
+		return UDID;
+	}
+
+private:
+    FString UDID;
+	FString Name;
+	FDeviceSyslogRelay*	syslogRelay;
+	FRunnableThread*	syslogRelayThread;
+};
+
+/**
+ * Delegate type for devices being connected or disconnected from the machine
+ *
+ * The first parameter is newly added or removed device
+ */
+DECLARE_MULTICAST_DELEGATE_OneParam(FDeviceNotification, void*)
+
+class FDeviceQueryTask
+	: public FRunnable
+{
+public:
+	FDeviceQueryTask()
+		: Stopping(false)
+		, bCheckDevices(true)
+	{}
+
+	virtual bool Init() override
+	{
+		return true;
+	}
+
+	virtual uint32 Run() override
+	{
+		while (!Stopping)
+		{
+			if (bCheckDevices)
+			{
+				// BHP - Turning off device check to prevent it from interfering with packaging
+				//QueryDevices();
+			}
+
+			FPlatformProcess::Sleep(5.0f);
+		}
+
+		return 0;
+	}
+
+	virtual void Stop() override
+	{
+		Stopping = true;
+	}
+
+	virtual void Exit() override
+	{}
+
+	FDeviceNotification& OnDeviceNotification()
+	{
+		return DeviceNotification;
+	}
+
+	void Enable(bool OnOff)
+	{
+		bCheckDevices = OnOff;
+	}
+
+private:
+	bool ExecuteDSCommand(const FString& CommandLine, FString* OutStdOut, FString* OutStdErr) const
+	{
+		FString ExecutablePath = FString::Printf(TEXT("%sBinaries/DotNET/IOS"), *FPaths::EngineDir());
+		FString Filename = TEXT("DeploymentServer.exe");
+
+		// execute the command
+		int32 ReturnCode;
+		FString DefaultError;
+
+		// make sure there's a place for error output to go if the caller specified nullptr
+		if (!OutStdErr)
+		{
+			OutStdErr = &DefaultError;
+		}
+
+		void* WritePipe;
+		void* ReadPipe;
+		FPlatformProcess::CreatePipe(ReadPipe, WritePipe);
+		FProcHandle ProcessHandle = FPlatformProcess::CreateProc(*(ExecutablePath / Filename), *CommandLine, false, true, true, NULL, 0, *ExecutablePath, WritePipe);
+		while (FPlatformProcess::IsProcRunning(ProcessHandle))
+		{
+			FString NewLine = FPlatformProcess::ReadPipe(ReadPipe);
+			if (NewLine.Len() > 0)
+			{
+				// process the string to break it up in to lines
+				*OutStdOut += NewLine;
+			}
+
+			FPlatformProcess::Sleep(0.25);
+		}
+		FString NewLine = FPlatformProcess::ReadPipe(ReadPipe);
+		if (NewLine.Len() > 0)
+		{
+			// process the string to break it up in to lines
+			*OutStdOut += NewLine;
+		}
+
+		FPlatformProcess::Sleep(0.25);
+		FPlatformProcess::ClosePipe(ReadPipe, WritePipe);
+
+		if (!FPlatformProcess::GetProcReturnCode(ProcessHandle, &ReturnCode))
+		{
+			return false;
+		}
+
+		if (ReturnCode != 0)
+		{
+			FPlatformMisc::LowLevelOutputDebugStringf(TEXT("The DeploymentServer command '%s' failed to run. Return code: %d, Error: %s\n"), *CommandLine, ReturnCode, **OutStdErr);
+
+			return false;
+		}
+
+		return true;
+	}
+
+	void QueryDevices()
+	{
+		FString StdOut;
+		// get the list of devices
+		if (!ExecuteDSCommand(TEXT("listdevices"), &StdOut, nullptr))
+		{
+			return;
+		}
+
+		// separate out each line
+		TArray<FString> DeviceStrings;
+		StdOut = StdOut.Replace(TEXT("\r"), TEXT("\n"));
+		StdOut.ParseIntoArray(DeviceStrings, TEXT("\n"), true);
+
+		TArray<FString> CurrentDeviceIds;
+		for (int32 StringIndex = 0; StringIndex < DeviceStrings.Num(); ++StringIndex)
+		{
+			const FString& DeviceString = DeviceStrings[StringIndex];
+
+			if(!DeviceString.StartsWith("[DD] FOUND: "))
+			{
+				continue;
+			}
+
+			// grab the device serial number
+			int32 idIndex = DeviceString.Find(TEXT("ID: "));
+			int32 nameIndex = DeviceString.Find(TEXT("NAME: "));
+			if (idIndex < 0 || nameIndex < 0)
+			{
+				continue;
+			}
+
+			FString SerialNumber = DeviceString.Mid(idIndex + 4, nameIndex - 1 - (idIndex + 4));
+			CurrentDeviceIds.Add(SerialNumber);
+
+			// move on to next device if this one is already a known device
+			if (ConnectedDeviceIds.Find(SerialNumber) != INDEX_NONE)
+			{
+				ConnectedDeviceIds.Remove(SerialNumber);
+				continue;
+			}
+
+			// parse device name
+			FString DeviceName = DeviceString.Mid(nameIndex + 6, DeviceString.Len() - (nameIndex + 6));
+
+			// create an FIOSDevice
+			FDeviceNotificationCallbackInformation CallbackInfo;
+			CallbackInfo.DeviceName = DeviceName;
+			CallbackInfo.UDID = SerialNumber;
+			CallbackInfo.msgType = 1;
+			DeviceNotification.Broadcast(&CallbackInfo);
+		}
+
+		// remove all devices no longer found
+		for (int32 DeviceIndex = 0; DeviceIndex < ConnectedDeviceIds.Num(); ++DeviceIndex)
+		{
+			FDeviceNotificationCallbackInformation CallbackInfo;
+			CallbackInfo.UDID = ConnectedDeviceIds[DeviceIndex];
+			CallbackInfo.msgType = 2;
+			DeviceNotification.Broadcast(&CallbackInfo);
+		}
+		ConnectedDeviceIds = CurrentDeviceIds;
+	}
+
+	bool Stopping;
+	bool bCheckDevices;
+	TArray<FString> ConnectedDeviceIds;
+	FDeviceNotification DeviceNotification;
+};
+
+/* FIOSDeviceHelper structors
+ *****************************************************************************/
+static TMap<FIOSDevice*, FIOSLaunchDaemonPong> ConnectedDevices;
+static FDeviceQueryTask* QueryTask = NULL;
+static FRunnableThread* QueryThread = NULL;
+static TArray<FDeviceNotificationCallbackInformation> NotificationMessages;
+static FTickerDelegate TickDelegate;
+
+bool FIOSDeviceHelper::MessageTickDelegate(float DeltaTime)
+{
+	for (int Index = 0; Index < NotificationMessages.Num(); ++Index)
+	{
+		FDeviceNotificationCallbackInformation cbi = NotificationMessages[Index];
+		FIOSDeviceHelper::DeviceCallback(&cbi);
+	}
+	NotificationMessages.Empty();
+
+	return true;
+}
+
+void FIOSDeviceHelper::Initialize(bool bIsTVOS)
+{
+	// Create a dummy device to hand over
+	const FString DummyDeviceName(FString::Printf(bIsTVOS ? TEXT("All_tvOS_On_%s") : TEXT("All_iOS_On_%s"), FPlatformProcess::ComputerName()));
+	
+	FIOSLaunchDaemonPong Event;
+	Event.DeviceID = FString::Printf(bIsTVOS ? TEXT("TVOS@%s") : TEXT("IOS@%s"), *DummyDeviceName);
+	Event.bCanReboot = false;
+	Event.bCanPowerOn = false;
+	Event.bCanPowerOff = false;
+	Event.DeviceName = DummyDeviceName;
+	Event.DeviceType = bIsTVOS ? TEXT("AppleTV") : TEXT("");
+	FIOSDeviceHelper::OnDeviceConnected().Broadcast(Event);
+
+	if(!bIsTVOS)
+	{
+		// add the message pump
+		TickDelegate = FTickerDelegate::CreateStatic(MessageTickDelegate);
+		FTicker::GetCoreTicker().AddTicker(TickDelegate, 5.0f);
+
+		// kick off a thread to query for connected devices
+		QueryTask = new FDeviceQueryTask();
+		QueryTask->OnDeviceNotification().AddStatic(FIOSDeviceHelper::DeviceCallback);
+		QueryThread = FRunnableThread::Create(QueryTask, TEXT("FIOSDeviceHelper.QueryTask"), 128 * 1024, TPri_Normal);
+	}
+}
+
+void FIOSDeviceHelper::DeviceCallback(void* CallbackInfo)
+{
+	struct FDeviceNotificationCallbackInformation* cbi = (FDeviceNotificationCallbackInformation*)CallbackInfo;
+
+	if (!IsInGameThread())
+	{
+		NotificationMessages.Add(*cbi);
+	}
+	else
+	{
+		switch(cbi->msgType)
+		{
+		case 1:
+			FIOSDeviceHelper::DoDeviceConnect(CallbackInfo);
+			break;
+
+		case 2:
+			FIOSDeviceHelper::DoDeviceDisconnect(CallbackInfo);
+			break;
+		}
+	}
+}
+
+void FIOSDeviceHelper::DoDeviceConnect(void* CallbackInfo)
+{
+	// connect to the device
+	struct FDeviceNotificationCallbackInformation* cbi = (FDeviceNotificationCallbackInformation*)CallbackInfo;
+	FIOSDevice* Device = new FIOSDevice(cbi->UDID, cbi->DeviceName);
+
+	// fire the event
+	FIOSLaunchDaemonPong Event;
+	Event.DeviceID = FString::Printf(TEXT("IOS@%s"), *(cbi->UDID));
+	Event.DeviceName = cbi->DeviceName;
+	Event.bCanReboot = false;
+	Event.bCanPowerOn = false;
+	Event.bCanPowerOff = false;
+	FIOSDeviceHelper::OnDeviceConnected().Broadcast(Event);
+
+	// add to the device list
+	ConnectedDevices.Add(Device, Event);
+}
+
+void FIOSDeviceHelper::DoDeviceDisconnect(void* CallbackInfo)
+{
+	struct FDeviceNotificationCallbackInformation* cbi = (FDeviceNotificationCallbackInformation*)CallbackInfo;
+	FIOSDevice* device = NULL;
+	for (auto DeviceIterator = ConnectedDevices.CreateIterator(); DeviceIterator; ++DeviceIterator)
+	{
+		if (DeviceIterator.Key()->SerialNumber() == cbi->UDID)
+		{
+			device = DeviceIterator.Key();
+			break;
+		}
+	}
+	if (device != NULL)
+	{
+		// extract the device id from the connected list
+		FIOSLaunchDaemonPong Event = ConnectedDevices.FindAndRemoveChecked(device);
+
+		// fire the event
+		FIOSDeviceHelper::OnDeviceDisconnected().Broadcast(Event);
+
+		// delete the device
+		delete device;
+	}
+}
+
+bool FIOSDeviceHelper::InstallIPAOnDevice(const FTargetDeviceId& DeviceId, const FString& IPAPath)
+{
+    return false;
+}
+
+void FIOSDeviceHelper::EnableDeviceCheck(bool OnOff)
+{
+	QueryTask->Enable(OnOff);
+}
\ No newline at end of file
diff --git a/Engine/Source/Developer/IOS/TVOSTargetPlatform/Private/TVOSTargetPlatformModule.cpp b/Engine/Source/Developer/IOS/TVOSTargetPlatform/Private/TVOSTargetPlatformModule.cpp
new file mode 100644
--- /dev/null
+++ b/Engine/Source/Developer/IOS/TVOSTargetPlatform/Private/TVOSTargetPlatformModule.cpp
@@ -0,0 +1,59 @@
+// Copyright 1998-2017 Epic Games, Inc. All Rights Reserved.
+
+#include "CoreMinimal.h"
+#include "IOSTargetPlatform.h"
+#include "Interfaces/ITargetPlatformModule.h"
+
+/**
+ * Holds the target platform singleton.
+ */
+static ITargetPlatform* Singleton = NULL;
+
+
+/**
+ * Module for TVOS as a target platform
+ */
+class FTVOSTargetPlatformModule : public ITargetPlatformModule
+{
+public:
+
+	/**
+	 * Destructor.
+	 */
+	~FTVOSTargetPlatformModule()
+	{
+		Singleton = NULL;
+	}
+
+
+public:
+
+	// Begin ITargetPlatformModule interface
+
+	virtual ITargetPlatform* GetTargetPlatform()
+	{
+		if (Singleton == NULL)
+		{
+			Singleton = new FIOSTargetPlatform(true);
+		}
+
+		return Singleton;
+	}
+
+	// End ITargetPlatformModule interface
+
+
+public:
+
+	// Begin IModuleInterface interface
+	virtual void StartupModule() override
+	{
+	}
+
+	virtual void ShutdownModule() override
+	{
+	}
+	// End IModuleInterface interface
+};
+
+IMPLEMENT_MODULE( FTVOSTargetPlatformModule, TVOSTargetPlatform);
diff --git a/Engine/Source/Developer/IOS/TVOSTargetPlatform/TVOSTargetPlatform.Build.cs b/Engine/Source/Developer/IOS/TVOSTargetPlatform/TVOSTargetPlatform.Build.cs
new file mode 100644
--- /dev/null
+++ b/Engine/Source/Developer/IOS/TVOSTargetPlatform/TVOSTargetPlatform.Build.cs
@@ -0,0 +1,44 @@
+// Copyright 1998-2017 Epic Games, Inc. All Rights Reserved.
+
+using UnrealBuildTool;
+
+public class TVOSTargetPlatform : ModuleRules
+{
+	public TVOSTargetPlatform(ReadOnlyTargetRules Target) : base(Target)
+	{
+		BinariesSubFolder = "IOS";
+		
+		PrivateDependencyModuleNames.AddRange(
+			new string[]
+			{
+				"Core",
+				"CoreUObject",
+				"TargetPlatform",
+				"DesktopPlatform",
+				"LaunchDaemonMessages",
+				"IOSTargetPlatform",
+				"Projects"
+			}
+		);
+
+		PrivateIncludePathModuleNames.AddRange(
+			new string[] {
+			"Messaging",
+			"TargetDeviceServices",
+		}
+		);
+
+		PlatformSpecificDynamicallyLoadedModuleNames.Add("LaunchDaemonMessages");
+
+		if (UEBuildConfiguration.bCompileAgainstEngine)
+		{
+			PrivateDependencyModuleNames.Add("Engine");
+		}
+
+		PrivateIncludePaths.AddRange(
+			new string[] {
+			"Developer/IOS/IOSTargetPlatform/Private"
+			}
+		);
+	}
+}
