fix for custom icons not causing a build with installed engine

diff --git a/Engine/Source/Developer/IOS/IOSPlatformEditor/IOSPlatformEditor.Build.cs b/Engine/Source/Developer/IOS/IOSPlatformEditor/IOSPlatformEditor.Build.cs
--- a/Engine/Source/Developer/IOS/IOSPlatformEditor/IOSPlatformEditor.Build.cs
+++ b/Engine/Source/Developer/IOS/IOSPlatformEditor/IOSPlatformEditor.Build.cs
@@ -26,9 +26,10 @@ public class IOSPlatformEditor : ModuleRules
 				"TargetPlatform",
 				"MaterialShaderQualitySettings",
 				"RenderCore",
-                "AudioSettingsEditor"
-			}
-		);
+                "AudioSettingsEditor",
+                "GameProjectGeneration"
+            }
+        );
 
 		PrivateIncludePathModuleNames.AddRange(
 			new string[] {
@@ -36,11 +37,5 @@ public class IOSPlatformEditor : ModuleRules
 				"Settings",
 			}
 		);
-
-		DynamicallyLoadedModuleNames.AddRange(
-			new string[] {
-				"GameProjectGeneration",
-				}
-		);
 	}
 }
diff --git a/Engine/Source/Developer/IOS/IOSPlatformEditor/Private/IOSCustomIconProjectBuildMutatorFeature.cpp b/Engine/Source/Developer/IOS/IOSPlatformEditor/Private/IOSCustomIconProjectBuildMutatorFeature.cpp
new file mode 100644
--- /dev/null
+++ b/Engine/Source/Developer/IOS/IOSPlatformEditor/Private/IOSCustomIconProjectBuildMutatorFeature.cpp
@@ -0,0 +1,65 @@
+// Copyright 1998-2018 Epic Games, Inc. All Rights Reserved.
+#include "IOSCustomIconProjectBuildMutatorFeature.h"
+#include "CoreMinimal.h"
+#include "HAL/PlatformFilemanager.h"
+#include "Misc/Paths.h"
+#include "PlatformInfo.h"
+#include "Interfaces/ITargetPlatform.h"
+#include "Interfaces/ITargetPlatformManagerModule.h"
+#include "GenericPlatform/GenericPlatformFile.h"
+#include "Misc/CoreMisc.h"
+
+static bool RequiresBuild()
+{
+	// determine if there are any project icons
+	FString IconDir = FPaths::Combine(FPaths::ProjectDir(), TEXT("Build/IOS/Resources/Graphics"));
+	struct FDirectoryVisitor : public IPlatformFile::FDirectoryVisitor
+	{
+		TArray<FString>& FileNames;
+
+		FDirectoryVisitor(TArray<FString>& InFileNames)
+			: FileNames(InFileNames)
+		{
+		}
+
+		virtual bool Visit(const TCHAR* FilenameOrDirectory, bool bIsDirectory) override
+		{
+			FString FileName(FilenameOrDirectory);
+			if (FileName.EndsWith(TEXT(".png")) && FileName.Contains(TEXT("Icon")))
+			{
+				FileNames.Add(FileName);
+			}
+			return true;
+		}
+	};
+
+	// Enumerate the contents of the current directory
+	TArray<FString> FileNames;
+	FDirectoryVisitor Visitor(FileNames);
+	FPlatformFileManager::Get().GetPlatformFile().IterateDirectory(*IconDir, Visitor);
+
+	if (FileNames.Num() > 0)
+	{
+		return true;
+	}
+	return false;
+}
+
+bool FIOSCustomIconProjectBuildMutatorFeature ::RequiresProjectBuild(FName InPlatformInfoName) const
+{
+	const PlatformInfo::FPlatformInfo* const PlatInfo = PlatformInfo::FindPlatformInfo(InPlatformInfoName);
+	check(PlatInfo);
+
+	if (PlatInfo->SDKStatus == PlatformInfo::EPlatformSDKStatus::Installed)
+	{
+		const ITargetPlatform* const Platform = GetTargetPlatformManager()->FindTargetPlatform(PlatInfo->TargetPlatformName.ToString());
+		if (Platform)
+		{
+			if (InPlatformInfoName.ToString() == TEXT("IOS"))
+			{
+				return RequiresBuild();
+			}
+		}
+	}
+	return false;
+}
\ No newline at end of file
diff --git a/Engine/Source/Developer/IOS/IOSPlatformEditor/Private/IOSCustomIconProjectBuildMutatorFeature.h b/Engine/Source/Developer/IOS/IOSPlatformEditor/Private/IOSCustomIconProjectBuildMutatorFeature.h
new file mode 100644
--- /dev/null
+++ b/Engine/Source/Developer/IOS/IOSPlatformEditor/Private/IOSCustomIconProjectBuildMutatorFeature.h
@@ -0,0 +1,14 @@
+// Copyright 1998-2018 Epic Games, Inc. All Rights Reserved.
+
+#pragma once
+
+#include "ProjectBuildMutatorFeature.h"
+
+class FIOSCustomIconProjectBuildMutatorFeature : public FProjectBuildMutatorFeature
+{
+public:
+
+	virtual bool RequiresProjectBuild(FName InPlatformInfoName) const override;
+
+private:
+};
\ No newline at end of file
diff --git a/Engine/Source/Developer/IOS/IOSPlatformEditor/Private/IOSPlatformEditorModule.cpp b/Engine/Source/Developer/IOS/IOSPlatformEditor/Private/IOSPlatformEditorModule.cpp
--- a/Engine/Source/Developer/IOS/IOSPlatformEditor/Private/IOSPlatformEditorModule.cpp
+++ b/Engine/Source/Developer/IOS/IOSPlatformEditor/Private/IOSPlatformEditorModule.cpp
@@ -13,6 +13,8 @@
 #include "MaterialShaderQualitySettingsCustomization.h"
 #include "ComponentRecreateRenderStateContext.h"
 #include "ShaderPlatformQualitySettings.h"
+#include "IOSCustomIconProjectBuildMutatorFeature.h"
+#include "Features/IModularFeatures.h"
 
 #define LOCTEXT_NAMESPACE "FIOSPlatformEditorModule"
 
@@ -70,6 +72,8 @@ class FIOSPlatformEditorModule
 				);
 			}
 		}
+
+		IModularFeatures::Get().RegisterModularFeature(FProjectBuildMutatorFeature::GetFeatureName(), &ProjectBuildMutator);
 	}
 
 	virtual void ShutdownModule() override
@@ -81,7 +85,14 @@ class FIOSPlatformEditorModule
 			SettingsModule->UnregisterSettings("Project", "Platforms", "iOS");
 			SettingsModule->UnregisterSettings("Project", "Platforms", "iOSMetalQuality");
 		}
+
+		IModularFeatures::Get().UnregisterModularFeature(FProjectBuildMutatorFeature::GetFeatureName(), &ProjectBuildMutator);
 	}
+
+private:
+
+	FIOSCustomIconProjectBuildMutatorFeature ProjectBuildMutator;
+
 };
 
 
