Replaced default iOS splashcreens with UE logo
Added missing resolutions (iPhone 5 portrait, iPad 10.5")
Removed custom splash images from RemoteSessionApp

diff --git a/Engine/Source/Developer/IOS/IOSPlatformEditor/Private/IOSTargetSettingsCustomization.cpp b/Engine/Source/Developer/IOS/IOSPlatformEditor/Private/IOSTargetSettingsCustomization.cpp
--- a/Engine/Source/Developer/IOS/IOSPlatformEditor/Private/IOSTargetSettingsCustomization.cpp
+++ b/Engine/Source/Developer/IOS/IOSPlatformEditor/Private/IOSTargetSettingsCustomization.cpp
@@ -85,23 +85,20 @@ FIOSTargetSettingsCustomization::FIOSTargetSettingsCustomization()
 	new (IconNames)FPlatformIconInfo(TEXT("[REDACTED_EMAIL]"), LOCTEXT("AppIcon_iPadProRetina_iOS9", "iPad Pro Retina iOS9 App Icon"), FText::GetEmpty(), 167, 167, FPlatformIconInfo::Required);
 	new (IconNames)FPlatformIconInfo(TEXT("Icon1024.png"), LOCTEXT("AppIcon_Marketing", "Marketing Icon"), FText::GetEmpty(), 1024, 1024, FPlatformIconInfo::Required);
 
-	new (LaunchImageNames) FPlatformIconInfo(TEXT("Default.png"), LOCTEXT("LaunchImage_iPhone", "Launch iPhone 4/4S"), FText::GetEmpty(), 320, 480, FPlatformIconInfo::Required);
-	new (LaunchImageNames) FPlatformIconInfo(TEXT("[REDACTED_EMAIL]"), LOCTEXT("LaunchImage_iPhoneRetina", "Launch iPhone 4/4S Retina"), FText::GetEmpty(), 640, 960, FPlatformIconInfo::Required);
-	new (LaunchImageNames) FPlatformIconInfo(TEXT("[REDACTED_EMAIL]"), LOCTEXT("LaunchImage_iPhone5", "Launch iPhone 5/5S Retina"), FText::GetEmpty(), 640, 1136, FPlatformIconInfo::Required);
-	new (LaunchImageNames) FPlatformIconInfo(TEXT("Default-Landscape.png"), LOCTEXT("LaunchImage_iPad_Landscape", "Launch iPad in Landscape"), FText::GetEmpty(), 1024, 768, FPlatformIconInfo::Required);
-	new (LaunchImageNames) FPlatformIconInfo(TEXT("[REDACTED_EMAIL]"), LOCTEXT("LaunchImage_iPadRetina_Landscape", "Launch iPad Retina in Landscape"), FText::GetEmpty(), 2048, 1536, FPlatformIconInfo::Required);
-	new (LaunchImageNames) FPlatformIconInfo(TEXT("Default-Portrait.png"), LOCTEXT("LaunchImage_iPad_Portrait", "Launch iPad in Portrait"), FText::GetEmpty(), 768, 1024, FPlatformIconInfo::Required);
-	new (LaunchImageNames) FPlatformIconInfo(TEXT("[REDACTED_EMAIL]"), LOCTEXT("LaunchImage_iPadRetina_Portrait", "Launch iPad Retina in Portrait"), FText::GetEmpty(), 1536, 2048, FPlatformIconInfo::Required);
-	new (LaunchImageNames) FPlatformIconInfo(TEXT("Default-IPhone6.png"), LOCTEXT("LaunchImage_iPhone6", "Launch iPhone 6 in Portrait"), FText::GetEmpty(), 750, 1334, FPlatformIconInfo::Required);
-	new (LaunchImageNames) FPlatformIconInfo(TEXT("Default-IPhone6-Landscape.png"), LOCTEXT("LaunchImage_iPhone6_Landscape", "Launch iPhone 6 in Landscape"), FText::GetEmpty(), 1334, 750, FPlatformIconInfo::Required);
-	new (LaunchImageNames) FPlatformIconInfo(TEXT("Default-IPhone6Plus-Landscape.png"), LOCTEXT("LaunchImage_iPhone6Plus_Landscape", "Launch iPhone 6 Plus in Landscape"), FText::GetEmpty(), 2208, 1242, FPlatformIconInfo::Required);
-	new (LaunchImageNames) FPlatformIconInfo(TEXT("Default-IPhone6Plus-Portrait.png"), LOCTEXT("LaunchImage_iPhone6Plus_Portrait", "Launch iPhone 6 Plus in Portrait"), FText::GetEmpty(), 1242, 2208, FPlatformIconInfo::Required);
-	new (LaunchImageNames) FPlatformIconInfo(TEXT("Default-IPhoneX-Landscape.png"), LOCTEXT("LaunchImage_iPhoneX_Landscape", "Launch iPhone X in Landscape"), FText::GetEmpty(), 2436, 1125, FPlatformIconInfo::Required);
-	new (LaunchImageNames) FPlatformIconInfo(TEXT("Default-IPhoneX-Portrait.png"), LOCTEXT("LaunchImage_iPhoneX_Portrait", "Launch iPhone X in Portrait"), FText::GetEmpty(), 1125, 2436, FPlatformIconInfo::Required);
-	new (LaunchImageNames) FPlatformIconInfo(TEXT("Default-Landscape-1336.png"), LOCTEXT("LaunchImage_iPadPro_Landscape", "Launch iPad Pro in Landscape"), FText::GetEmpty(), 1336, 1024, FPlatformIconInfo::Required);
-	new (LaunchImageNames)FPlatformIconInfo(TEXT("Default-Portrait-1336.png"), LOCTEXT("LaunchImage_iPadPro_Portrait", "Launch iPad Pro in Portrait"), FText::GetEmpty(), 1024, 1336, FPlatformIconInfo::Required);
-	new (LaunchImageNames)FPlatformIconInfo(TEXT("[REDACTED_EMAIL]"), LOCTEXT("LaunchImage_iPadProRetina_Landscape", "Launch iPad Pro Retina in Landscape"), FText::GetEmpty(), 2732, 2048, FPlatformIconInfo::Required);
-	new (LaunchImageNames)FPlatformIconInfo(TEXT("[REDACTED_EMAIL]"), LOCTEXT("LaunchImage_iPadProRetina_Portrait", "Launch iPad Pro Retina in Portrait"), FText::GetEmpty(), 2048, 2732, FPlatformIconInfo::Required);
+	new (LaunchImageNames) FPlatformIconInfo(TEXT("Default-IPhone5-Landscape.png"), LOCTEXT("LaunchImage_iPhone5_Landscape", "iPhone 5/5S/SE in Landscape"), FText::GetEmpty(), 1136, 640, FPlatformIconInfo::Required);
+    new (LaunchImageNames) FPlatformIconInfo(TEXT("[REDACTED_EMAIL]"), LOCTEXT("LaunchImage_iPhone5_Portrait", "iPhone 5/5S/SE in Portrait"), FText::GetEmpty(), 640, 1136, FPlatformIconInfo::Required);
+    new (LaunchImageNames) FPlatformIconInfo(TEXT("Default-IPhone6-Landscape.png"), LOCTEXT("LaunchImage_iPhone6_Landscape", "iPhone 6/7/8 in Landscape"), FText::GetEmpty(), 1334, 750, FPlatformIconInfo::Required);
+    new (LaunchImageNames) FPlatformIconInfo(TEXT("Default-IPhone6.png"), LOCTEXT("LaunchImage_iPhone6", "iPhone 6/7/8 in Portrait"), FText::GetEmpty(), 750, 1334, FPlatformIconInfo::Required);
+    new (LaunchImageNames) FPlatformIconInfo(TEXT("Default-IPhone6Plus-Landscape.png"), LOCTEXT("LaunchImage_iPhone6Plus_Landscape", "iPhone 6+/7+/8+ in Landscape"), FText::GetEmpty(), 2208, 1242, FPlatformIconInfo::Required);
+    new (LaunchImageNames) FPlatformIconInfo(TEXT("Default-IPhone6Plus-Portrait.png"), LOCTEXT("LaunchImage_iPhone6Plus_Portrait", "iPhone 6+/7+/8+ in Portrait"), FText::GetEmpty(), 1242, 2208, FPlatformIconInfo::Required);
+    new (LaunchImageNames) FPlatformIconInfo(TEXT("Default-IPhoneX-Landscape.png"), LOCTEXT("LaunchImage_iPhoneX_Landscape", "iPhone X in Landscape"), FText::GetEmpty(), 2436, 1125, FPlatformIconInfo::Required);
+    new (LaunchImageNames) FPlatformIconInfo(TEXT("Default-IPhoneX-Portrait.png"), LOCTEXT("LaunchImage_iPhoneX_Portrait", "iPhone X in Portrait"), FText::GetEmpty(), 1125, 2436, FPlatformIconInfo::Required);
+	new (LaunchImageNames) FPlatformIconInfo(TEXT("[REDACTED_EMAIL]"), LOCTEXT("LaunchImage_iPadRetina_Landscape", "iPad in Landscape"), FText::GetEmpty(), 2048, 1536, FPlatformIconInfo::Required);
+	new (LaunchImageNames) FPlatformIconInfo(TEXT("[REDACTED_EMAIL]"), LOCTEXT("LaunchImage_iPadRetina_Portrait", "iPad in Portrait"), FText::GetEmpty(), 1536, 2048, FPlatformIconInfo::Required);
+	new (LaunchImageNames) FPlatformIconInfo(TEXT("Default-IPadPro10_5_Landscape.png"), LOCTEXT("LaunchImage_iPadPro_Landscape", "iPad Pro 10.5 in Landscape"), FText::GetEmpty(), 2224, 1668, FPlatformIconInfo::Required);
+	new (LaunchImageNames) FPlatformIconInfo(TEXT("Default-IPadPro10_5_Portrait.png"), LOCTEXT("LaunchImage_iPadPro_Portrait", "iPad Pro 10.5 in Portrait"), FText::GetEmpty(), 1668, 2224, FPlatformIconInfo::Required);
+	new (LaunchImageNames)FPlatformIconInfo(TEXT("[REDACTED_EMAIL]"), LOCTEXT("LaunchImage_iPadProRetina_Landscape", "iPad Pro 12.9 in Landscape"), FText::GetEmpty(), 2732, 2048, FPlatformIconInfo::Required);
+	new (LaunchImageNames)FPlatformIconInfo(TEXT("[REDACTED_EMAIL]"), LOCTEXT("LaunchImage_iPadProRetina_Portrait", "iPad Pro 12.9 in Portrait"), FText::GetEmpty(), 2048, 2732, FPlatformIconInfo::Required);
 
 	bShowAllProvisions = false;
 	bShowAllCertificates = false;
