UEMOB-465 - Add Device Output Log support for iOS

diff --git a/Engine/Source/Developer/IOS/IOSTargetPlatform/Private/IOSTargetDevice.cpp b/Engine/Source/Developer/IOS/IOSTargetPlatform/Private/IOSTargetDevice.cpp
--- a/Engine/Source/Developer/IOS/IOSTargetPlatform/Private/IOSTargetDevice.cpp
+++ b/Engine/Source/Developer/IOS/IOSTargetPlatform/Private/IOSTargetDevice.cpp
@@ -7,7 +7,7 @@
 #include "Interfaces/ITargetPlatform.h"
 #include "MessageEndpoint.h"
 #include "MessageEndpointBuilder.h"
-
+#include "IOSTargetDeviceOutput.h"
 
 FIOSTargetDevice::FIOSTargetDevice(const ITargetPlatform& InTargetPlatform)
 	: TargetPlatform(InTargetPlatform)
@@ -189,3 +189,13 @@ bool FIOSTargetDevice::GetUserCredentials(FString& OutUserName, FString& OutUser
 {
 	return false;
 }
+inline ITargetDeviceOutputPtr FIOSTargetDevice::CreateDeviceOutputRouter(FOutputDevice* Output) const
+{
+	FIOSTargetDeviceOutputPtr DeviceOutputPtr = MakeShareable(new FIOSTargetDeviceOutput());
+	if (DeviceOutputPtr->Init(*this, Output))
+	{
+		return DeviceOutputPtr;
+	}
+
+	return nullptr;
+}
diff --git a/Engine/Source/Developer/IOS/IOSTargetPlatform/Private/IOSTargetDevice.h b/Engine/Source/Developer/IOS/IOSTargetPlatform/Private/IOSTargetDevice.h
--- a/Engine/Source/Developer/IOS/IOSTargetPlatform/Private/IOSTargetDevice.h
+++ b/Engine/Source/Developer/IOS/IOSTargetPlatform/Private/IOSTargetDevice.h
@@ -16,6 +16,8 @@ typedef TSharedPtr<class FIOSTargetDevice, ESPMode::ThreadSafe> FIOSTargetDevice
 /** Type definition for shared references to instances of FIOSTargetDevice. */
 typedef TSharedRef<class FIOSTargetDevice, ESPMode::ThreadSafe> FIOSTargetDeviceRef;
 
+/** Type definition for shared references to instances of FIOSTargetDeviceOutput. */
+typedef TSharedPtr<class FIOSTargetDeviceOutput, ESPMode::ThreadSafe> FIOSTargetDeviceOutputPtr;
 
 /**
  * Implements an iOS target device.
@@ -57,6 +59,7 @@ public:
 	virtual bool TerminateProcess(const int64 ProcessId) override;
 	virtual void SetUserCredentials(const FString& UserName, const FString& UserPassword) override;
 	virtual bool GetUserCredentials(FString& OutUserName, FString& OutUserPassword) override;
+	virtual ITargetDeviceOutputPtr CreateDeviceOutputRouter(FOutputDevice* Output) const override;
 
 public:
 
diff --git a/Engine/Source/Developer/IOS/IOSTargetPlatform/Private/Windows/IOSDeviceHelperWindows.cpp b/Engine/Source/Developer/IOS/IOSTargetPlatform/Private/Windows/IOSDeviceHelperWindows.cpp
--- a/Engine/Source/Developer/IOS/IOSTargetPlatform/Private/Windows/IOSDeviceHelperWindows.cpp
+++ b/Engine/Source/Developer/IOS/IOSTargetPlatform/Private/Windows/IOSDeviceHelperWindows.cpp
@@ -176,10 +176,10 @@ public:
 			if (bCheckDevices)
 			{
 				// BHP - Turning off device check to prevent it from interfering with packaging
-				//QueryDevices();
+				QueryDevices();
 			}
 
-			FPlatformProcess::Sleep(5.0f);
+			FPlatformProcess::Sleep(10.0f);
 		}
 
 		return 0;
diff --git a/Engine/Source/Developer/OutputLog/Private/SDeviceOutputLog.cpp b/Engine/Source/Developer/OutputLog/Private/SDeviceOutputLog.cpp
--- a/Engine/Source/Developer/OutputLog/Private/SDeviceOutputLog.cpp
+++ b/Engine/Source/Developer/OutputLog/Private/SDeviceOutputLog.cpp
@@ -19,10 +19,11 @@ static bool IsSupportedPlatform(ITargetPlatform* Platform)
 {
 	static const FName AndroidPlaftomName("Android"); // TODO: currently implemented only for Android & Switch. #TODO: perhaps move to ETargetDeviceFeatures ?
 	static const FName SwitchPlatformName("Switch");
+	static const FName IOSPlatformName("IOS");
 	
 	check(Platform);
 	const auto& PlatfromInfo = Platform->GetPlatformInfo();
-	return PlatfromInfo.IsVanilla() && (PlatfromInfo.VanillaPlatformName == AndroidPlaftomName || PlatfromInfo.VanillaPlatformName == SwitchPlatformName);
+	return PlatfromInfo.IsVanilla() && (PlatfromInfo.VanillaPlatformName == AndroidPlaftomName || PlatfromInfo.VanillaPlatformName == SwitchPlatformName || PlatfromInfo.VanillaPlatformName == IOSPlatformName);
 }
 
 
@@ -238,6 +239,21 @@ void SDeviceOutputLog::HandleTargetPlatformDeviceDiscovered(ITargetDeviceRef Dis
 
 void SDeviceOutputLog::AddDeviceEntry(ITargetDeviceRef TargetDevice)
 {
+	if (FindDeviceEntry(TargetDevice->GetId()))
+	{
+		return;
+	}
+	if (TargetDevice->IsPlatformAggregated())
+	{
+		return;
+	}
+	const FString DummyIOSDeviceName(FString::Printf(TEXT("All_iOS_On_%s"), FPlatformProcess::ComputerName()));
+	const FString DummyTVOSDeviceName(FString::Printf(TEXT("All_tvOS_On_%s"), FPlatformProcess::ComputerName()));
+	if (TargetDevice->GetId().GetDeviceName().Equals(DummyIOSDeviceName, ESearchCase::IgnoreCase) ||
+		TargetDevice->GetId().GetDeviceName().Equals(DummyTVOSDeviceName, ESearchCase::IgnoreCase))
+	{
+		return;
+	}
 	using namespace PlatformInfo;
 	FName DeviceIconStyleName = TargetDevice->GetTargetPlatform().GetPlatformInfo().GetIconStyleName(EPlatformIconSize::Normal);
 	
@@ -250,6 +266,19 @@ void SDeviceOutputLog::AddDeviceEntry(ITargetDeviceRef TargetDevice)
 	DeviceList.Add(DeviceEntry);
 }
 
+bool SDeviceOutputLog::FindDeviceEntry(FTargetDeviceId InDeviceId)
+{
+	// Should not do it, but what if someone somewhere holds strong reference to a lost device?
+	for (const TSharedPtr<FTargetDeviceEntry>& EntryPtr : DeviceList)
+	{
+		if (EntryPtr->DeviceId.GetDeviceName() == InDeviceId.GetDeviceName())
+		{
+			return true;
+		}
+	}
+	return false;
+}
+
 void SDeviceOutputLog::OnDeviceSelectionChanged(FTargetDeviceEntryPtr DeviceEntry)
 {
 	CurrentDeviceOutputPtr.Reset();
diff --git a/Engine/Source/Developer/OutputLog/Private/SDeviceOutputLog.h b/Engine/Source/Developer/OutputLog/Private/SDeviceOutputLog.h
--- a/Engine/Source/Developer/OutputLog/Private/SDeviceOutputLog.h
+++ b/Engine/Source/Developer/OutputLog/Private/SDeviceOutputLog.h
@@ -53,7 +53,8 @@ protected:
 	/** Callback for discovered target devices. */
 	void HandleTargetPlatformDeviceDiscovered(ITargetDeviceRef DiscoveredDevice);
 
-	void AddDeviceEntry(ITargetDeviceRef TargetDevice);
+	void  AddDeviceEntry(ITargetDeviceRef TargetDevice);
+	bool FindDeviceEntry(FTargetDeviceId InDeviceId);
 
 	void OnDeviceSelectionChanged(FTargetDeviceEntryPtr DeviceEntry);
 	TSharedRef<SWidget> MakeDeviceComboButtonMenu();
diff --git a/Engine/Source/Developer/iOS/IOSTargetPlatform/Private/IOSTargetDeviceOutput.h b/Engine/Source/Developer/iOS/IOSTargetPlatform/Private/IOSTargetDeviceOutput.h
new file mode 100644
--- /dev/null
+++ b/Engine/Source/Developer/iOS/IOSTargetPlatform/Private/IOSTargetDeviceOutput.h
@@ -0,0 +1,55 @@
+// Copyright 1998-2018 Epic Games, Inc. All Rights Reserved.
+
+#pragma once
+
+#include "Containers/UnrealString.h"
+#include "HAL/Runnable.h"
+#include "HAL/RunnableThread.h"
+#include "HAL/PlatformProcess.h"
+#include "Interfaces/ITargetDeviceOutput.h"
+#include "Misc/ConfigCacheIni.h"
+#include "HAL/ThreadSafeCounter.h"
+
+class FIOSTargetDevice;
+
+class FIOSDeviceOutputReaderRunnable : public FRunnable
+{
+public:
+	FIOSDeviceOutputReaderRunnable(const FTargetDeviceId InDeviceId, FOutputDevice* Output);
+	
+	// FRunnable interface.
+	virtual bool Init(void) override;
+	virtual void Exit(void) override; 
+	virtual void Stop(void) override;
+	virtual uint32 Run(void) override;
+
+private:
+	bool StartDSProcess(void);
+
+private:
+	// > 0 if we've been asked to abort work in progress at the next opportunity
+	FThreadSafeCounter	StopTaskCounter;
+	
+	FTargetDeviceId		DeviceId;
+	FOutputDevice*		Output;
+	
+	void*				DSReadPipe;
+	void*				DSWritePipe;
+	FProcHandle			DSProcHandle;
+};
+
+/**
+ * Implements a IOS target device output.
+ */
+class FIOSTargetDeviceOutput : public ITargetDeviceOutput
+{
+public:
+	bool Init(const FIOSTargetDevice& TargetDevice, FOutputDevice* Output);
+	
+private:
+	TUniquePtr<FRunnableThread>						DeviceOutputThread;
+	FTargetDeviceId									DeviceId;
+	FString											DeviceName;
+};
+
+#include "IOSTargetDeviceOutput.inl"
diff --git a/Engine/Source/Developer/iOS/IOSTargetPlatform/Private/IOSTargetDeviceOutput.inl b/Engine/Source/Developer/iOS/IOSTargetPlatform/Private/IOSTargetDeviceOutput.inl
new file mode 100644
--- /dev/null
+++ b/Engine/Source/Developer/iOS/IOSTargetPlatform/Private/IOSTargetDeviceOutput.inl
@@ -0,0 +1,125 @@
+// Copyright 1998-2018 Epic Games, Inc. All Rights Reserved.
+
+#include "CoreTypes.h"
+#include "CoreFwd.h"
+
+class FIOSDeviceOutputReaderRunnable;
+class FIOSTargetDevice;
+class FIOSTargetDeviceOutput;
+
+inline FIOSDeviceOutputReaderRunnable::FIOSDeviceOutputReaderRunnable(const FTargetDeviceId InDeviceId, FOutputDevice* InOutput)
+	: StopTaskCounter(0)
+	, DeviceId(InDeviceId)
+	, Output(InOutput)
+	, DSReadPipe(nullptr)
+	, DSWritePipe(nullptr)
+{
+}
+
+inline bool FIOSDeviceOutputReaderRunnable::StartDSProcess(void)
+{
+	FString DSFilename = FPaths::ConvertRelativePathToFull(FPaths::EngineDir() / TEXT("Binaries/DotNET/IOS/DeploymentServer.exe"));
+	FString Params = FString::Printf(TEXT(" listentodevice -device %s"), *DeviceId.GetDeviceName());
+	Output->Serialize(TEXT("Starting process ....."), ELogVerbosity::Log, NAME_None);
+	Output->Serialize(*DeviceId.GetDeviceName(), ELogVerbosity::Log, NAME_None);
+	DSProcHandle = FPlatformProcess::CreateProc(*DSFilename, *Params, true, false, false, NULL, 0, NULL, DSWritePipe);
+	return DSProcHandle.IsValid();
+}
+
+inline bool FIOSDeviceOutputReaderRunnable::Init(void) 
+{ 
+	FPlatformProcess::CreatePipe(DSReadPipe, DSWritePipe);
+	return StartDSProcess();
+}
+
+inline void FIOSDeviceOutputReaderRunnable::Exit(void) 
+{
+	StopTaskCounter.Increment();
+	if (DSProcHandle.IsValid())
+	{
+		if (DSReadPipe && DSWritePipe)
+		{
+			FPlatformProcess::ClosePipe(DSReadPipe, DSWritePipe);
+			DSReadPipe = nullptr;
+			DSWritePipe = nullptr;
+		}
+		FPlatformProcess::TerminateProc(DSProcHandle, true);
+	}
+}
+
+inline void FIOSDeviceOutputReaderRunnable::Stop(void)
+{
+	StopTaskCounter.Increment();
+}
+
+inline uint32 FIOSDeviceOutputReaderRunnable::Run(void)
+{
+	FString DSOutput;
+	
+	while (StopTaskCounter.GetValue() == 0 && DSProcHandle.IsValid())
+	{
+		if (!FPlatformProcess::IsProcRunning(DSProcHandle))
+		{
+			// When user plugs out USB cable DS process stops
+			// Keep trying to restore DS connection until code that uses this object will not kill us
+			Output->Serialize(TEXT("Trying to restore connection to device..."), ELogVerbosity::Log, NAME_None);
+			FPlatformProcess::CloseProc(DSProcHandle);
+			if (StartDSProcess())
+			{
+				FPlatformProcess::Sleep(5.0f);
+			}
+			else
+			{
+				Output->Serialize(TEXT("Failed to start DS proccess"), ELogVerbosity::Log, NAME_None);
+				Stop();
+			}
+		}
+		else
+		{
+			DSOutput.Append(FPlatformProcess::ReadPipe(DSReadPipe));
+
+			if (DSOutput.Len() > 0)
+			{
+				TArray<FString> OutputLines;
+				DSOutput.ParseIntoArray(OutputLines, TEXT("\n"), false);
+
+				if (!DSOutput.EndsWith(TEXT("\n")))
+				{
+					// partial line at the end, do not serialize it until we receive remainder
+					DSOutput = OutputLines.Last();
+					OutputLines.RemoveAt(OutputLines.Num() - 1);
+				}
+				else
+				{
+					DSOutput.Reset();
+				}
+
+				for (int32 i = 0; i < OutputLines.Num(); ++i)
+				{
+					Output->Serialize(*OutputLines[i], ELogVerbosity::Log, NAME_None);
+				}
+			}
+			
+			FPlatformProcess::Sleep(0.1f);
+		}
+	}
+
+	return 0;
+}
+
+inline bool FIOSTargetDeviceOutput::Init(const FIOSTargetDevice& TargetDevice, FOutputDevice* Output)
+{
+	check(Output);
+	// Output will be produced by background thread
+	check(Output->CanBeUsedOnAnyThread());
+	
+	DeviceId = TargetDevice.GetId();
+	DeviceName = TargetDevice.GetName();
+
+	Output->Serialize(TEXT("Creating FIOSTargetDeviceOutput ....."), ELogVerbosity::Log, NAME_None);
+		
+	auto* Runnable = new FIOSDeviceOutputReaderRunnable(DeviceId, Output);
+	DeviceOutputThread = TUniquePtr<FRunnableThread>(FRunnableThread::Create(Runnable, TEXT("FIOSDeviceOutputReaderRunnable")));
+	return true;
+}
+	
