Disable iOS device detection on Windows, required due to backing out deployment server

diff --git a/Engine/Source/Developer/IOS/IOSTargetPlatform/Private/Windows/IOSDeviceHelperWindows.cpp b/Engine/Source/Developer/IOS/IOSTargetPlatform/Private/Windows/IOSDeviceHelperWindows.cpp
--- a/Engine/Source/Developer/IOS/IOSTargetPlatform/Private/Windows/IOSDeviceHelperWindows.cpp
+++ b/Engine/Source/Developer/IOS/IOSTargetPlatform/Private/Windows/IOSDeviceHelperWindows.cpp
@@ -176,10 +176,10 @@ public:
 			if (bCheckDevices)
 			{
 				// BHP - Turning off device check to prevent it from interfering with packaging
-				QueryDevices();
+				//QueryDevices();
 			}
 
-			FPlatformProcess::Sleep(10.0f);
+			FPlatformProcess::Sleep(5.0f);
 		}
 
 		return 0;
