Re-adding with IOS to fix case issue

diff --git a/Engine/Source/Developer/iOS/IOSTargetPlatform/Private/IOSTargetDeviceOutput.h b/Engine/Source/Developer/iOS/IOSTargetPlatform/Private/IOSTargetDeviceOutput.h
new file mode 100644
--- /dev/null
+++ b/Engine/Source/Developer/iOS/IOSTargetPlatform/Private/IOSTargetDeviceOutput.h
@@ -0,0 +1,55 @@
+// Copyright 1998-2018 Epic Games, Inc. All Rights Reserved.
+
+#pragma once
+
+#include "Containers/UnrealString.h"
+#include "HAL/Runnable.h"
+#include "HAL/RunnableThread.h"
+#include "HAL/PlatformProcess.h"
+#include "Interfaces/ITargetDeviceOutput.h"
+#include "Misc/ConfigCacheIni.h"
+#include "HAL/ThreadSafeCounter.h"
+
+class FIOSTargetDevice;
+
+class FIOSDeviceOutputReaderRunnable : public FRunnable
+{
+public:
+	FIOSDeviceOutputReaderRunnable(const FTargetDeviceId InDeviceId, FOutputDevice* Output);
+	
+	// FRunnable interface.
+	virtual bool Init(void) override;
+	virtual void Exit(void) override; 
+	virtual void Stop(void) override;
+	virtual uint32 Run(void) override;
+
+private:
+	bool StartDSProcess(void);
+
+private:
+	// > 0 if we've been asked to abort work in progress at the next opportunity
+	FThreadSafeCounter	StopTaskCounter;
+	
+	FTargetDeviceId		DeviceId;
+	FOutputDevice*		Output;
+	
+	void*				DSReadPipe;
+	void*				DSWritePipe;
+	FProcHandle			DSProcHandle;
+};
+
+/**
+ * Implements a IOS target device output.
+ */
+class FIOSTargetDeviceOutput : public ITargetDeviceOutput
+{
+public:
+	bool Init(const FIOSTargetDevice& TargetDevice, FOutputDevice* Output);
+	
+private:
+	TUniquePtr<FRunnableThread>						DeviceOutputThread;
+	FTargetDeviceId									DeviceId;
+	FString											DeviceName;
+};
+
+#include "IOSTargetDeviceOutput.inl"
diff --git a/Engine/Source/Developer/iOS/IOSTargetPlatform/Private/IOSTargetDeviceOutput.inl b/Engine/Source/Developer/iOS/IOSTargetPlatform/Private/IOSTargetDeviceOutput.inl
new file mode 100644
--- /dev/null
+++ b/Engine/Source/Developer/iOS/IOSTargetPlatform/Private/IOSTargetDeviceOutput.inl
@@ -0,0 +1,125 @@
+// Copyright 1998-2018 Epic Games, Inc. All Rights Reserved.
+
+#include "CoreTypes.h"
+#include "CoreFwd.h"
+
+class FIOSDeviceOutputReaderRunnable;
+class FIOSTargetDevice;
+class FIOSTargetDeviceOutput;
+
+inline FIOSDeviceOutputReaderRunnable::FIOSDeviceOutputReaderRunnable(const FTargetDeviceId InDeviceId, FOutputDevice* InOutput)
+	: StopTaskCounter(0)
+	, DeviceId(InDeviceId)
+	, Output(InOutput)
+	, DSReadPipe(nullptr)
+	, DSWritePipe(nullptr)
+{
+}
+
+inline bool FIOSDeviceOutputReaderRunnable::StartDSProcess(void)
+{
+	FString DSFilename = FPaths::ConvertRelativePathToFull(FPaths::EngineDir() / TEXT("Binaries/DotNET/IOS/DeploymentServer.exe"));
+	FString Params = FString::Printf(TEXT(" listentodevice -device %s"), *DeviceId.GetDeviceName());
+	Output->Serialize(TEXT("Starting process ....."), ELogVerbosity::Log, NAME_None);
+	Output->Serialize(*DeviceId.GetDeviceName(), ELogVerbosity::Log, NAME_None);
+	DSProcHandle = FPlatformProcess::CreateProc(*DSFilename, *Params, true, false, false, NULL, 0, NULL, DSWritePipe);
+	return DSProcHandle.IsValid();
+}
+
+inline bool FIOSDeviceOutputReaderRunnable::Init(void) 
+{ 
+	FPlatformProcess::CreatePipe(DSReadPipe, DSWritePipe);
+	return StartDSProcess();
+}
+
+inline void FIOSDeviceOutputReaderRunnable::Exit(void) 
+{
+	StopTaskCounter.Increment();
+	if (DSProcHandle.IsValid())
+	{
+		if (DSReadPipe && DSWritePipe)
+		{
+			FPlatformProcess::ClosePipe(DSReadPipe, DSWritePipe);
+			DSReadPipe = nullptr;
+			DSWritePipe = nullptr;
+		}
+		FPlatformProcess::TerminateProc(DSProcHandle, true);
+	}
+}
+
+inline void FIOSDeviceOutputReaderRunnable::Stop(void)
+{
+	StopTaskCounter.Increment();
+}
+
+inline uint32 FIOSDeviceOutputReaderRunnable::Run(void)
+{
+	FString DSOutput;
+	
+	while (StopTaskCounter.GetValue() == 0 && DSProcHandle.IsValid())
+	{
+		if (!FPlatformProcess::IsProcRunning(DSProcHandle))
+		{
+			// When user plugs out USB cable DS process stops
+			// Keep trying to restore DS connection until code that uses this object will not kill us
+			Output->Serialize(TEXT("Trying to restore connection to device..."), ELogVerbosity::Log, NAME_None);
+			FPlatformProcess::CloseProc(DSProcHandle);
+			if (StartDSProcess())
+			{
+				FPlatformProcess::Sleep(5.0f);
+			}
+			else
+			{
+				Output->Serialize(TEXT("Failed to start DS proccess"), ELogVerbosity::Log, NAME_None);
+				Stop();
+			}
+		}
+		else
+		{
+			DSOutput.Append(FPlatformProcess::ReadPipe(DSReadPipe));
+
+			if (DSOutput.Len() > 0)
+			{
+				TArray<FString> OutputLines;
+				DSOutput.ParseIntoArray(OutputLines, TEXT("\n"), false);
+
+				if (!DSOutput.EndsWith(TEXT("\n")))
+				{
+					// partial line at the end, do not serialize it until we receive remainder
+					DSOutput = OutputLines.Last();
+					OutputLines.RemoveAt(OutputLines.Num() - 1);
+				}
+				else
+				{
+					DSOutput.Reset();
+				}
+
+				for (int32 i = 0; i < OutputLines.Num(); ++i)
+				{
+					Output->Serialize(*OutputLines[i], ELogVerbosity::Log, NAME_None);
+				}
+			}
+			
+			FPlatformProcess::Sleep(0.1f);
+		}
+	}
+
+	return 0;
+}
+
+inline bool FIOSTargetDeviceOutput::Init(const FIOSTargetDevice& TargetDevice, FOutputDevice* Output)
+{
+	check(Output);
+	// Output will be produced by background thread
+	check(Output->CanBeUsedOnAnyThread());
+	
+	DeviceId = TargetDevice.GetId();
+	DeviceName = TargetDevice.GetName();
+
+	Output->Serialize(TEXT("Creating FIOSTargetDeviceOutput ....."), ELogVerbosity::Log, NAME_None);
+		
+	auto* Runnable = new FIOSDeviceOutputReaderRunnable(DeviceId, Output);
+	DeviceOutputThread = TUniquePtr<FRunnableThread>(FRunnableThread::Create(Runnable, TEXT("FIOSDeviceOutputReaderRunnable")));
+	return true;
+}
+	
