UE-70234 - Error: Failed to connect to deployment server at 127.0.0.1:41000 - DeploymentServer not running

diff --git a/Engine/Source/Developer/IOS/IOSTargetPlatform/Private/IOSTargetDevice.cpp b/Engine/Source/Developer/IOS/IOSTargetPlatform/Private/IOSTargetDevice.cpp
--- a/Engine/Source/Developer/IOS/IOSTargetPlatform/Private/IOSTargetDevice.cpp
+++ b/Engine/Source/Developer/IOS/IOSTargetPlatform/Private/IOSTargetDevice.cpp
@@ -22,6 +22,7 @@ FTcpDSCommander::FTcpDSCommander(const uint8* Data, int32 Count, void* WPipe)
 : bStopping(false)
 , bStoped(true)
 , bIsSuccess(false)
+, bIsSystemError(false)
 , DSSocket(nullptr)
 , Thread(nullptr)
 , WritePipe(WPipe)
@@ -29,9 +30,7 @@ FTcpDSCommander::FTcpDSCommander(const uint8* Data, int32 Count, void* WPipe)
 , DSCommandLen(Count + 1)
 , LastActivity(0.0)
 {
-	// create the socket
-	ISocketSubsystem* SSS = ISocketSubsystem::Get();
-    if (SSS && Count > 0)
+    if (Count > 0)
     {
         DSCommand = (uint8*)FMemory::Malloc(Count + 1);
         FPlatformMemory::Memcpy(DSCommand, Data, Count);
@@ -62,36 +61,60 @@ void FTcpDSCommander::Exit()
 
 bool FTcpDSCommander::Init()
 {
-    if (DSCommandLen < 1)
-    {
-        bIsSuccess = true;
-        return true;
-    }
+	if (DSCommandLen < 1)
+	{
+		bIsSuccess = true;
+		return true;
+	}
 	ISocketSubsystem* SSS = ISocketSubsystem::Get();
-    DSSocket = SSS->CreateSocket(NAME_Stream, TEXT("DSCommander tcp"));
-    if (DSSocket == nullptr)
-    {
-        return false;
-    }
-    TSharedRef<FInternetAddr> Addr = SSS->CreateInternetAddr(0, DEFAULT_DS_COMMANDER_PORT);
-    bool bIsValid;
-    Addr->SetIp(TEXT("127.0.0.1"), bIsValid);
-    
-    
-    // try to connect to the server
-    if (DSSocket->Connect(*Addr) == false)
-    {
-        StartDSProcess();
-        if (DSSocket->Connect(*Addr) == false)
-        {
-            // on failure, shut it all down
-            SSS->DestroySocket(DSSocket);
-            DSSocket = nullptr;
+	DSSocket = SSS->CreateSocket(NAME_Stream, TEXT("DSCommander tcp"));
+	if (DSSocket == nullptr)
+	{
+		return false;
+	}
+	TSharedRef<FInternetAddr> Addr = SSS->CreateInternetAddr(0, DEFAULT_DS_COMMANDER_PORT);
+	bool bIsValid;
+	Addr->SetIp(TEXT("127.0.0.1"), bIsValid);
+
+#if PLATFORM_WINDOWS
+	// using the mutex to detect if the DeploymentServer is running
+	// only on windows
+	if (!IsDSRunning())
+	{
+		StartDSProcess();
+
+		int TimeoutSeconds = 5;
+		while (!IsDSRunning() && TimeoutSeconds > 0)
+		{
+			TimeoutSeconds--;
+			FPlatformProcess::Sleep(1.0f);
+		}
+		if (!IsDSRunning())
+		{
+			bIsSystemError = true;
+			return false;
+		}
+	}
+	if (DSSocket->Connect(*Addr) == false)
+	{
+		{ // extra bracked for cleaner ifdefs
+#else
+	// try to connect to the server
+	// on mac we use the old way to try a TCP connection
+	if (DSSocket->Connect(*Addr) == false)
+	{
+		StartDSProcess();
+		if (DSSocket->Connect(*Addr) == false)
+		{
+#endif // PLATFORM_WINDOWS
+        	// on failure, shut it all down
+        	SSS->DestroySocket(DSSocket);
+        	DSSocket = nullptr;
 			ESocketErrors LastError = SSS->GetLastErrorCode();
 			const TCHAR* SocketErr = SSS->GetSocketError(LastError);
-            UE_LOG(LogTemp, Display, TEXT("Failed to connect to deployment server at %s (%s)."), *Addr->ToString(true), SocketErr);
-            return false;
-        }
+        	//UE_LOG(LogTemp, Display, TEXT("Failed to connect to deployment server at %s (%s)."), *Addr->ToString(true), SocketErr);
+        	return false;
+		}
     }
     LastActivity = FPlatformTime::Seconds();
     
@@ -102,7 +125,7 @@ uint32 FTcpDSCommander::Run()
 {
     if (!DSSocket)
     {
-        UE_LOG(LogTemp, Log, TEXT("Socket not created."));
+        //UE_LOG(LogTemp, Log, TEXT("Socket not created."));
         return 1;
     }
     int32 NSent = 0;
@@ -110,7 +133,7 @@ uint32 FTcpDSCommander::Run()
     if (NSent != DSCommandLen || !BSent)
     {
         Stop();
-        UE_LOG(LogTemp, Log, TEXT("Socket send error."));
+        //UE_LOG(LogTemp, Log, TEXT("Socket send error."));
         return 1;
     }
     bStoped = false;
@@ -124,7 +147,7 @@ uint32 FTcpDSCommander::Run()
         if (DSSocket->GetConnectionState() != ESocketConnectionState::SCS_Connected)
         {
             Stop();
-            UE_LOG(LogTemp, Log, TEXT("Socket connection error."));
+            //UE_LOG(LogTemp, Log, TEXT("Socket connection error."));
             return 1;
         }
         if (DSSocket->HasPendingData(Pending))
@@ -154,7 +177,7 @@ uint32 FTcpDSCommander::Run()
                     else if (TagArray[i].EndsWith(TEXT("CMDFAIL\r")))
                     {
                         Stop();
-                        UE_LOG(LogTemp, Log, TEXT("Socket command failed."));
+                        //UE_LOG(LogTemp, Display, TEXT("Socket command failed."));
                         return 1;
                     }
                     else
@@ -165,10 +188,11 @@ uint32 FTcpDSCommander::Run()
             }
         }
         double CurrentTime = FPlatformTime::Seconds();
-        if (CurrentTime - LastActivity > 10.0)
+        if (CurrentTime - LastActivity > 120.0)
         {
+			//UE_LOG(LogTemp, Display, TEXT("Socket command timeouted."));
             Stop();
-            return 1;
+            return 0;
         }
         FPlatformProcess::Sleep(0.01f);
     }
@@ -180,6 +204,7 @@ void FTcpDSCommander::Stop()
 {
     if (DSSocket)
     {
+		DSSocket->Shutdown(ESocketShutdownMode::ReadWrite);
         DSSocket->Close();
         ISocketSubsystem::Get()->DestroySocket(DSSocket);
     }
@@ -188,21 +213,26 @@ void FTcpDSCommander::Stop()
     bStopping = true;
 }
 
-bool FTcpDSCommander::StartDSProcess()
+bool FTcpDSCommander::IsDSRunning()
 {
-    // is there a mutex we can use to connect test DS is running, also available on mac?
-    // there is also a failsafe mechanism for this, since the DeploymentServer will not start a new server if one is already running
+	// is there a mutex we can use to connect test DS is running, also available on mac?
+	// there is also a failsafe mechanism for this, since the DeploymentServer will not start a new server if one is already running
 #if PLATFORM_WINDOWS
-    HANDLE mutex = CreateMutexA(NULL, TRUE, "Global\\DeploymentServer_Mutex_SERVERINSTANCE");
-    if (mutex == NULL)
-    {
-        // deployment server instance already runnning
-		UE_LOG(LogTemp, Display, TEXT("Deployment server instance already runnning."));
-        return false;
-    }
-    CloseHandle(mutex);
+	HANDLE mutex = CreateMutexA(NULL, TRUE, "Global\\DeploymentServer_Mutex_SERVERINSTANCE");
+	if (mutex == NULL || GetLastError() == ERROR_ALREADY_EXISTS)
+	{
+		// deployment server instance already runnning
+		//UE_LOG(LogTemp, Display, TEXT("Deployment server instance already runnning."));
+		CloseHandle(mutex);
+		return true;
+	}
+	CloseHandle(mutex);
 #endif // PLATFORM_WINDOWS
-    
+	return false;
+}
+
+bool FTcpDSCommander::StartDSProcess()
+{   
     FString DSFilename = FPaths::ConvertRelativePathToFull(FPaths::EngineDir() / TEXT("Binaries/DotNET/IOS/DeploymentServerLauncher.exe"));
     FString WorkingFoder = FPaths::ConvertRelativePathToFull(FPaths::EngineDir() / TEXT("Binaries/DotNET/IOS/"));
     FString Params = "";
@@ -213,7 +243,7 @@ bool FTcpDSCommander::StartDSProcess()
     Params = FString::Printf(TEXT("\"%s\" \"%s\" %s"), *ScriptPath, *DSFilename, *Params);
     DSFilename = TEXT("/bin/sh");
 #endif
-    UE_LOG(LogTemp, Log, TEXT("DeploymentServer not running, Starting it!"));
+    //UE_LOG(LogTemp, Log, TEXT("DeploymentServer not running, Starting it!"));
     FPlatformProcess::CreateProc(*DSFilename, *Params, true, true, true, NULL, 0, *WorkingFoder, (void*)nullptr);
     FPlatformProcess::Sleep(1.0f);
     
@@ -422,7 +452,7 @@ inline ITargetDeviceOutputPtr FIOSTargetDevice::CreateDeviceOutputRouter(FOutput
 	return nullptr;
 }
 
-bool FIOSTargetDeviceOutput::ExecuteDSCommand(const char *CommandLine, FString* OutStdOut)
+int FIOSTargetDeviceOutput::ExecuteDSCommand(const char *CommandLine, FString* OutStdOut)
 {
 	void* WritePipe;
 	void* ReadPipe;
@@ -449,12 +479,16 @@ bool FIOSTargetDeviceOutput::ExecuteDSCommand(const char *CommandLine, FString*
 	FPlatformProcess::Sleep(0.25);
 	FPlatformProcess::ClosePipe(ReadPipe, WritePipe);
 
-	if (!DSCommander.WasSuccess())
+	if (DSCommander.IsSystemError())
 	{
-		FPlatformMisc::LowLevelOutputDebugStringf(TEXT("The DeploymentServer command '%s' failed to run.\n"), ANSI_TO_TCHAR(CommandLine));
+		return -1;
+	}
 
-		return false;
+	if (!DSCommander.WasSuccess())
+	{
+		//FPlatformMisc::LowLevelOutputDebugStringf(TEXT("The DeploymentServer command '%s' failed to run.\n"), ANSI_TO_TCHAR(CommandLine));
+		return 0;
 	}
 
-	return true;
+	return 1;
 }
diff --git a/Engine/Source/Developer/IOS/IOSTargetPlatform/Private/IOSTargetDevice.h b/Engine/Source/Developer/IOS/IOSTargetPlatform/Private/IOSTargetDevice.h
--- a/Engine/Source/Developer/IOS/IOSTargetPlatform/Private/IOSTargetDevice.h
+++ b/Engine/Source/Developer/IOS/IOSTargetPlatform/Private/IOSTargetDevice.h
@@ -36,6 +36,9 @@ public:
     
     /** Virtual destructor. */
     virtual ~FTcpDSCommander();
+
+	/** Check if DeploymentServer mutex is active*/
+	static bool IsDSRunning();
     
 public:
     
@@ -58,6 +61,10 @@ public:
     {
         return bIsSuccess;
     }
+	inline bool IsSystemError()
+	{
+		return bIsSystemError;
+	}
     
 private:
     
@@ -67,6 +74,7 @@ private:
     bool bStopping;
     bool bStoped;
     bool bIsSuccess;
+	bool bIsSystemError; ///< Deployment server was not able to start, or connection to it could not be made
     
     /** */
     class FSocket* DSSocket;
diff --git a/Engine/Source/Developer/IOS/IOSTargetPlatform/Private/IOSTargetDeviceOutput.h b/Engine/Source/Developer/IOS/IOSTargetPlatform/Private/IOSTargetDeviceOutput.h
--- a/Engine/Source/Developer/IOS/IOSTargetPlatform/Private/IOSTargetDeviceOutput.h
+++ b/Engine/Source/Developer/IOS/IOSTargetPlatform/Private/IOSTargetDeviceOutput.h
@@ -46,7 +46,7 @@ class FIOSTargetDeviceOutput : public ITargetDeviceOutput
 public:
 	bool Init(const FIOSTargetDevice& TargetDevice, FOutputDevice* Output);
 
-	static bool ExecuteDSCommand(const char *CommandLine, FString* OutStdOut);
+	static int ExecuteDSCommand(const char *CommandLine, FString* OutStdOut);
 	
 private:
 	TUniquePtr<FRunnableThread>						DeviceOutputThread;
diff --git a/Engine/Source/Developer/IOS/IOSTargetPlatform/Private/Windows/IOSDeviceHelperWindows.cpp b/Engine/Source/Developer/IOS/IOSTargetPlatform/Private/Windows/IOSDeviceHelperWindows.cpp
--- a/Engine/Source/Developer/IOS/IOSTargetPlatform/Private/Windows/IOSDeviceHelperWindows.cpp
+++ b/Engine/Source/Developer/IOS/IOSTargetPlatform/Private/Windows/IOSDeviceHelperWindows.cpp
@@ -5,6 +5,8 @@
 #include "HAL/PlatformProcess.h"
 #include "HAL/Runnable.h"
 #include "HAL/RunnableThread.h"
+#include "Interfaces/ITargetPlatformManagerModule.h"
+#include "Interfaces/ITargetPlatform.h"
 
 struct FDeviceNotificationCallbackInformation
 {
@@ -51,6 +53,7 @@ public:
 	FDeviceQueryTask()
 		: Stopping(false)
 		, bCheckDevices(true)
+		, NeedSDKCheck(true)
 		, RetryQuery(5)
 	{}
 
@@ -68,8 +71,32 @@ public:
 #if WITH_EDITOR
 				if (!IsRunningCommandlet())
 				{
-					// BHP - Turning off device check to prevent it from interfering with packaging
-					QueryDevices();
+					if (NeedSDKCheck)
+					{
+						if (GetTargetPlatformManager())
+						{
+							bool CanQuery = false;
+							FString OutTutorialPath;
+							const ITargetPlatform* Platform = GetTargetPlatformManager()->FindTargetPlatform(TEXT("IOS"));
+							if (Platform)
+							{
+								if (Platform->IsSdkInstalled(false, OutTutorialPath))
+								{
+									CanQuery = true;
+								}
+							}
+							NeedSDKCheck = false;
+							if (!CanQuery)
+							{
+								Enable(false);
+							}
+						}
+					}
+					else
+					{
+						// BHP - Turning off device check to prevent it from interfering with packaging
+						QueryDevices();
+					}
 				}
 #endif
 			}
@@ -104,16 +131,18 @@ private:
 	{
 		FString StdOut;
 		// get the list of devices
-		if (!FIOSTargetDeviceOutput::ExecuteDSCommand("listdevices", &StdOut))
+		int Response = FIOSTargetDeviceOutput::ExecuteDSCommand("listdevices", &StdOut);
+		if (Response <= 0)
 		{
 			RetryQuery--;
-			if (RetryQuery < 0)
+			if (RetryQuery < 0 || Response < 0)
 			{
+				UE_LOG(LogTemp, Log, TEXT("IOS device listing is disabled (to many failed attempts)!"));
 				Enable(false);
 			}
 			return;
 		}
-
+		RetryQuery = 5;
 		// separate out each line
 		TArray<FString> DeviceStrings;
 		StdOut = StdOut.Replace(TEXT("\r"), TEXT("\n"));
@@ -174,6 +203,7 @@ private:
 
 	bool Stopping;
 	bool bCheckDevices;
+	bool NeedSDKCheck;
 	int RetryQuery;
 	TArray<FString> ConnectedDeviceIds;
 	FDeviceNotification DeviceNotification;
@@ -226,7 +256,12 @@ void FIOSDeviceHelper::Initialize(bool bIsTVOS)
 		QueryTask->OnDeviceNotification().AddStatic(FIOSDeviceHelper::DeviceCallback);
 
 		static int32 QueryTaskCount = 1;
-		QueryThread = FRunnableThread::Create(QueryTask, *FString::Printf(TEXT("FIOSDeviceHelper.QueryTask_%d"), QueryTaskCount++), 128 * 1024, TPri_Normal);
+		if (QueryTaskCount == 1)
+		{
+			// create the socket subsystem (loadmodule in game thread)
+			ISocketSubsystem* SSS = ISocketSubsystem::Get();
+			QueryThread = FRunnableThread::Create(QueryTask, *FString::Printf(TEXT("FIOSDeviceHelper.QueryTask_%d"), QueryTaskCount++), 128 * 1024, TPri_Normal);
+		}
 	}
 }
 
