UE-70550 - //UE4/Release-4.22 - UE4Editor Static Analysis Win64 (MSVC) - warning C6387: 'mutex' could be '0': this does not adhere to the specification for the function 'CloseHandle' - iostargetdevice.cpp

diff --git a/Engine/Source/Developer/IOS/IOSTargetPlatform/Private/IOSTargetDevice.cpp b/Engine/Source/Developer/IOS/IOSTargetPlatform/Private/IOSTargetDevice.cpp
--- a/Engine/Source/Developer/IOS/IOSTargetPlatform/Private/IOSTargetDevice.cpp
+++ b/Engine/Source/Developer/IOS/IOSTargetPlatform/Private/IOSTargetDevice.cpp
@@ -222,8 +222,10 @@ bool FTcpDSCommander::IsDSRunning()
 	if (mutex == NULL || GetLastError() == ERROR_ALREADY_EXISTS)
 	{
 		// deployment server instance already runnning
-		//UE_LOG(LogTemp, Display, TEXT("Deployment server instance already runnning."));
-		CloseHandle(mutex);
+		if (mutex)
+		{
+			CloseHandle(mutex);
+		}
 		return true;
 	}
 	CloseHandle(mutex);
