UE-62851 - Remote Build pre-check error does not appear when missing required arguments

diff --git a/Engine/Source/Developer/IOS/IOSTargetPlatform/Private/IOSTargetPlatform.cpp b/Engine/Source/Developer/IOS/IOSTargetPlatform/Private/IOSTargetPlatform.cpp
--- a/Engine/Source/Developer/IOS/IOSTargetPlatform/Private/IOSTargetPlatform.cpp
+++ b/Engine/Source/Developer/IOS/IOSTargetPlatform/Private/IOSTargetPlatform.cpp
@@ -262,8 +262,10 @@ int32 FIOSTargetPlatform::CheckRequirements(const FString& ProjectPath, bool bPr
 	FString CmdExe = FPaths::ConvertRelativePathToFull(FPaths::EngineDir() / TEXT("Binaries/DotNET/IOS/IPhonePackager.exe"));
 	FString CommandLine = FString::Printf(TEXT("Validate Engine -project \"%s\" -bundlename \"%s\" %s"), *ProjectPath, *(BundleIdentifier), (bForDistribtion ? TEXT("-distribution") : TEXT("")) );
 	FString RemoteServerName;
+	FString RSyncUsername;
 	GConfig->GetString(TEXT("/Script/IOSRuntimeSettings.IOSRuntimeSettings"), TEXT("RemoteServerName"), RemoteServerName, GEngineIni);
-	if (RemoteServerName.Len() == 0)
+	GConfig->GetString(TEXT("/Script/IOSRuntimeSettings.IOSRuntimeSettings"), TEXT("RSyncUsername"), RSyncUsername, GEngineIni);
+	if (RemoteServerName.Len() == 0 || RSyncUsername.Len() == 0)
 	{
 		bReadyToBuild |= ETargetPlatformReadyStatus::RemoveServerNameEmpty;
 	}
