UE-70653 - Edit hangs for a few seconds when entering iOS commands in device output log

diff --git a/Engine/Source/Developer/IOS/IOSTargetPlatform/Private/IOSTargetDevice.cpp b/Engine/Source/Developer/IOS/IOSTargetPlatform/Private/IOSTargetDevice.cpp
--- a/Engine/Source/Developer/IOS/IOSTargetPlatform/Private/IOSTargetDevice.cpp
+++ b/Engine/Source/Developer/IOS/IOSTargetPlatform/Private/IOSTargetDevice.cpp
@@ -5,6 +5,7 @@
 #include "HAL/PlatformProcess.h"
 #include "IOSMessageProtocol.h"
 #include "Interfaces/ITargetPlatform.h"
+#include "Async/Async.h"
 #include "MessageEndpoint.h"
 #include "MessageEndpointBuilder.h"
 #include "IOSTargetDeviceOutput.h"
@@ -438,9 +439,16 @@ bool FIOSTargetDevice::GetUserCredentials(FString& OutUserName, FString& OutUser
 
 inline void FIOSTargetDevice::ExecuteConsoleCommand(const FString& ExecCommand) const
 {
-	FString StdOut;
 	FString Params = FString::Printf(TEXT("command -device %s -param \"%s\""), *DeviceId.GetDeviceName(), *ExecCommand);
-	FIOSTargetDeviceOutput::ExecuteDSCommand(TCHAR_TO_ANSI(*Params), &StdOut);
+
+	AsyncTask(
+		ENamedThreads::AnyThread,
+		[Params]()
+		{
+			FString StdOut;
+			FIOSTargetDeviceOutput::ExecuteDSCommand(TCHAR_TO_ANSI(*Params), &StdOut);
+		}
+	);
 }
 
 inline ITargetDeviceOutputPtr FIOSTargetDevice::CreateDeviceOutputRouter(FOutputDevice* Output) const
