UE-73418 - DevelopmentServerLauncher.exe loops whenever the editor is open even if no iOS devices are connected

diff --git a/Engine/Source/Developer/IOS/IOSTargetPlatform/Private/Windows/IOSDeviceHelperWindows.cpp b/Engine/Source/Developer/IOS/IOSTargetPlatform/Private/Windows/IOSDeviceHelperWindows.cpp
--- a/Engine/Source/Developer/IOS/IOSTargetPlatform/Private/Windows/IOSDeviceHelperWindows.cpp
+++ b/Engine/Source/Developer/IOS/IOSTargetPlatform/Private/Windows/IOSDeviceHelperWindows.cpp
@@ -7,6 +7,7 @@
 #include "HAL/RunnableThread.h"
 #include "Interfaces/ITargetPlatformManagerModule.h"
 #include "Interfaces/ITargetPlatform.h"
+#include "Interfaces/IProjectManager.h"
 
 struct FDeviceNotificationCallbackInformation
 {
@@ -46,6 +47,9 @@ private:
  */
 DECLARE_MULTICAST_DELEGATE_OneParam(FDeviceNotification, void*)
 
+// recheck once per minute
+#define		RECHECK_COUNTER_RESET			12
+
 class FDeviceQueryTask
 	: public FRunnable
 {
@@ -64,6 +68,7 @@ public:
 
 	virtual uint32 Run() override
 	{
+		int RecheckCounter = RECHECK_COUNTER_RESET;
 		while (!Stopping)
 		{
 			if (bCheckDevices)
@@ -90,6 +95,14 @@ public:
 							{
 								Enable(false);
 							}
+							else
+							{
+								FProjectStatus ProjectStatus;
+								if (!IProjectManager::Get().QueryStatusForCurrentProject(ProjectStatus) || (!ProjectStatus.IsTargetPlatformSupported(TEXT("IOS")) && !ProjectStatus.IsTargetPlatformSupported(TEXT("TVOS"))))
+								{
+									Enable(false);
+								}
+							}
 						}
 					}
 					else
@@ -100,6 +113,13 @@ public:
 				}
 #endif
 			}
+			RecheckCounter--;
+			if (RecheckCounter < 0)
+			{
+				RecheckCounter = RECHECK_COUNTER_RESET;
+				bCheckDevices = true;
+				NeedSDKCheck = true;
+			}
 
 			FPlatformProcess::Sleep(5.0f);
 		}
