Change FPlatformMisc::GetEnvironmentVariable() to return an FString() rather than requiring a fixed size buffer to be passed in. Removes references to MAX_PATH.

diff --git a/Engine/Plugins/Developer/GitSourceControl/Source/GitSourceControl/Private/GitSourceControlUtils.cpp b/Engine/Plugins/Developer/GitSourceControl/Source/GitSourceControl/Private/GitSourceControlUtils.cpp
--- a/Engine/Plugins/Developer/GitSourceControl/Source/GitSourceControl/Private/GitSourceControlUtils.cpp
+++ b/Engine/Plugins/Developer/GitSourceControl/Source/GitSourceControl/Private/GitSourceControlUtils.cpp
@@ -164,9 +164,8 @@ FString FindGitBinaryPath()
 	if(!bFound)
 	{
 		// else the install dir for the current user: [REDACTED_USER_HOME]\AppData\Local\Programs\Git\cmd
-		TCHAR AppDataLocalPath[4096];
-		FPlatformMisc::GetEnvironmentVariable(TEXT("LOCALAPPDATA"), AppDataLocalPath, ARRAY_COUNT(AppDataLocalPath));
-		GitBinaryPath = FString::Printf(TEXT("%s/Programs/Git/cmd/git.exe"), AppDataLocalPath);
+		FString AppDataLocalPath = FPlatformMisc::GetEnvironmentVariable(TEXT("LOCALAPPDATA"));
+		GitBinaryPath = FString::Printf(TEXT("%s/Programs/Git/cmd/git.exe"), *AppDataLocalPath);
 		bFound = CheckGitAvailability(GitBinaryPath);
 	}
 
@@ -187,9 +186,8 @@ FString FindGitBinaryPath()
 	if(!bFound)
 	{
 		// [REDACTED_USER_HOME]\AppData\Local\Atlassian\SourceTree\git_local\bin
-		TCHAR AppDataLocalPath[4096];
-		FPlatformMisc::GetEnvironmentVariable(TEXT("LOCALAPPDATA"), AppDataLocalPath, ARRAY_COUNT(AppDataLocalPath));
-		GitBinaryPath = FString::Printf(TEXT("%s/Atlassian/SourceTree/git_local/bin/git.exe"), AppDataLocalPath);
+		FString AppDataLocalPath = FPlatformMisc::GetEnvironmentVariable(TEXT("LOCALAPPDATA"));
+		GitBinaryPath = FString::Printf(TEXT("%s/Atlassian/SourceTree/git_local/bin/git.exe"), *AppDataLocalPath);
 		bFound = CheckGitAvailability(GitBinaryPath);
 	}
 
@@ -198,20 +196,19 @@ FString FindGitBinaryPath()
 	{
 		// The latest GitHub Desktop adds its binaries into the local appdata directory:
 		// [REDACTED_USER_HOME]\AppData\Local\GitHub\PortableGit_c2ba306e536fdf878271f7fe636a147ff37326ad\cmd
-		TCHAR AppDataLocalPath[4096];
-		FPlatformMisc::GetEnvironmentVariable(TEXT("LOCALAPPDATA"), AppDataLocalPath, ARRAY_COUNT(AppDataLocalPath));
-		FString SearchPath = FString::Printf(TEXT("%s/GitHub/PortableGit_*"), AppDataLocalPath);
+		FString AppDataLocalPath = FPlatformMisc::GetEnvironmentVariable(TEXT("LOCALAPPDATA"));
+		FString SearchPath = FString::Printf(TEXT("%s/GitHub/PortableGit_*"), *AppDataLocalPath);
 		TArray<FString> PortableGitFolders;
 		IFileManager::Get().FindFiles(PortableGitFolders, *SearchPath, false, true);
 		if(PortableGitFolders.Num() > 0)
 		{
 			// FindFiles just returns directory names, so we need to prepend the root path to get the full path.
-			GitBinaryPath = FString::Printf(TEXT("%s/GitHub/%s/cmd/git.exe"), AppDataLocalPath, *(PortableGitFolders.Last())); // keep only the last PortableGit found
+			GitBinaryPath = FString::Printf(TEXT("%s/GitHub/%s/cmd/git.exe"), *AppDataLocalPath, *(PortableGitFolders.Last())); // keep only the last PortableGit found
 			bFound = CheckGitAvailability(GitBinaryPath);
 			if (!bFound)
 			{
 				// If Portable git is not found in "cmd/" subdirectory, try the "bin/" path that was in use before
-				GitBinaryPath = FString::Printf(TEXT("%s/GitHub/%s/bin/git.exe"), AppDataLocalPath, *(PortableGitFolders.Last())); // keep only the last PortableGit found
+				GitBinaryPath = FString::Printf(TEXT("%s/GitHub/%s/bin/git.exe"), *AppDataLocalPath, *(PortableGitFolders.Last())); // keep only the last PortableGit found
 				bFound = CheckGitAvailability(GitBinaryPath);
 			}
 		}
diff --git a/Engine/Plugins/Developer/PerforceSourceControl/Source/PerforceSourceControl/Private/PerforceConnection.cpp b/Engine/Plugins/Developer/PerforceSourceControl/Source/PerforceSourceControl/Private/PerforceConnection.cpp
--- a/Engine/Plugins/Developer/PerforceSourceControl/Source/PerforceSourceControl/Private/PerforceConnection.cpp
+++ b/Engine/Plugins/Developer/PerforceSourceControl/Source/PerforceSourceControl/Private/PerforceConnection.cpp
@@ -484,9 +484,7 @@ bool FPerforceConnection::GetWorkspaceList(const FPerforceConnectionInfo& InConn
 			if(LocalHostName.Len() == 0)
 			{
 				// No host override, check environment variable
-				TCHAR P4HostEnv[256];
-				FPlatformMisc::GetEnvironmentVariable(TEXT("P4HOST"), P4HostEnv, ARRAY_COUNT(P4HostEnv));
-				LocalHostName = P4HostEnv;
+				LocalHostName = FPlatformMisc::GetEnvironmentVariable(TEXT("P4HOST"));
 			}
 
 			if (LocalHostName.Len() == 0)
diff --git a/Engine/Plugins/Experimental/PythonScriptPlugin/Source/PythonScriptPluginPreload/Private/PythonScriptPluginPreload.cpp b/Engine/Plugins/Experimental/PythonScriptPlugin/Source/PythonScriptPluginPreload/Private/PythonScriptPluginPreload.cpp
--- a/Engine/Plugins/Experimental/PythonScriptPlugin/Source/PythonScriptPluginPreload/Private/PythonScriptPluginPreload.cpp
+++ b/Engine/Plugins/Experimental/PythonScriptPlugin/Source/PythonScriptPluginPreload/Private/PythonScriptPluginPreload.cpp
@@ -63,13 +63,7 @@ private:
 			if (PythonDLLPaths.Num() == 0)
 			{
 				// If we didn't find anything, check the Windows directory as the DLLs can sometimes be installed there
-				FString WinDir;
-				{
-					TCHAR WinDirA[1024] = {0};
-					FPlatformMisc::GetEnvironmentVariable(TEXT("WINDIR"), WinDirA, ARRAY_COUNT(WinDirA) - 1);
-					WinDir = WinDirA;
-				}
-
+				FString WinDir = FPlatformMisc::GetEnvironmentVariable(TEXT("WINDIR"));
 				if (!WinDir.IsEmpty())
 				{
 					PythonDLLPaths = FindPythonDLLs(WinDir / TEXT("System32"));
diff --git a/Engine/Plugins/Experimental/Shotgun/Source/Shotgun/Private/ShotgunUIManager.cpp b/Engine/Plugins/Experimental/Shotgun/Source/Shotgun/Private/ShotgunUIManager.cpp
--- a/Engine/Plugins/Experimental/Shotgun/Source/Shotgun/Private/ShotgunUIManager.cpp
+++ b/Engine/Plugins/Experimental/Shotgun/Source/Shotgun/Private/ShotgunUIManager.cpp
@@ -51,24 +51,15 @@ void FShotgunUIManagerImpl::Initialize()
 {
 	bIsShotgunEnabled = false;
 
-	TCHAR EnvVariable[4096];
-
 	// Check if the bootstrap environment variable is set and that the script exists
-	FPlatformMisc::GetEnvironmentVariable(TEXT("UE_SHOTGUN_BOOTSTRAP"), EnvVariable, ARRAY_COUNT(EnvVariable));
-	FString ShotgunBootstrap(EnvVariable);
-
+	FString ShotgunBootstrap = FPlatformMisc::GetEnvironmentVariable(TEXT("UE_SHOTGUN_BOOTSTRAP"));
 	if (!ShotgunBootstrap.IsEmpty() && FPaths::FileExists(ShotgunBootstrap))
 	{
 		// The following environment variables must be set for the Shotgun apps to be fully functional
 		// These variables are automatically set when the editor is launched through Shotgun Desktop
-		FPlatformMisc::GetEnvironmentVariable(TEXT("SHOTGUN_ENGINE"), EnvVariable, ARRAY_COUNT(EnvVariable));
-		FString ShotgunEngine(EnvVariable);
-
-		FPlatformMisc::GetEnvironmentVariable(TEXT("SHOTGUN_ENTITY_TYPE"), EnvVariable, ARRAY_COUNT(EnvVariable));
-		FString ShotgunEntityType(EnvVariable);
-
-		FPlatformMisc::GetEnvironmentVariable(TEXT("SHOTGUN_ENTITY_ID"), EnvVariable, ARRAY_COUNT(EnvVariable));
-		FString ShotgunEntityId(EnvVariable);
+		FString ShotgunEngine = FPlatformMisc::GetEnvironmentVariable(TEXT("SHOTGUN_ENGINE"));
+		FString ShotgunEntityType = FPlatformMisc::GetEnvironmentVariable(TEXT("SHOTGUN_ENTITY_TYPE"));
+		FString ShotgunEntityId = FPlatformMisc::GetEnvironmentVariable(TEXT("SHOTGUN_ENTITY_ID"));
 
 		if (ShotgunEngine == TEXT("tk-unreal") && !ShotgunEntityType.IsEmpty() && !ShotgunEntityId.IsEmpty())
 		{
diff --git a/Engine/Plugins/Lumin/MagicLeap/Source/MLSDK/Public/MagicLeapPluginUtil.h b/Engine/Plugins/Lumin/MagicLeap/Source/MLSDK/Public/MagicLeapPluginUtil.h
--- a/Engine/Plugins/Lumin/MagicLeap/Source/MLSDK/Public/MagicLeapPluginUtil.h
+++ b/Engine/Plugins/Lumin/MagicLeap/Source/MLSDK/Public/MagicLeapPluginUtil.h
@@ -28,7 +28,6 @@ public:
 		// We search various places for the ML API DLLs to support loading alternate
 		// implementations. For example to use VDZI on PC platforms.
 		// Public MLSDK path.
-		FString MLSDK = FString::ChrN(1024, TEXT('\0'));
 
 		// Give preference to the config setting.
 		FString MLSDKPath;
@@ -56,8 +55,7 @@ public:
 			}
 		}
 
-		FPlatformMisc::GetEnvironmentVariable(TEXT("MLSDK"), MLSDK.GetCharArray().GetData(), 1024);
-		MLSDK.TrimToNullTerminator();
+		FString MLSDK = FPlatformMisc::GetEnvironmentVariable(TEXT("MLSDK"));
 
 		if (CheckForVDZILibraries)
 		{
@@ -104,10 +102,7 @@ public:
 #if PLATFORM_WINDOWS
 		// Need to adjust PATH with additional MLSDK load path to allow the delay-loaded DLLs
 		// to work in the plugin.
-		const int32 MaxPathVarLen = 32768;
-		FString PathVar = FString::ChrN(MaxPathVarLen, TEXT('\0'));
-		FPlatformMisc::GetEnvironmentVariable(TEXT("PATH"), PathVar.GetCharArray().GetData(), MaxPathVarLen);
-		PathVar.TrimToNullTerminator();
+		FString PathVar = FPlatformMisc::GetEnvironmentVariable(TEXT("PATH"));
 		for (const FString& path : DllSearchPaths)
 		{
 			PathVar.Append(FPlatformMisc::GetPathVarDelimiter());
diff --git a/Engine/Plugins/Runtime/GoogleVR/GoogleVRHMD/Source/GoogleVRHMD/Public/GoogleVRAdbUtils.h b/Engine/Plugins/Runtime/GoogleVR/GoogleVRHMD/Source/GoogleVRHMD/Public/GoogleVRAdbUtils.h
--- a/Engine/Plugins/Runtime/GoogleVR/GoogleVRHMD/Source/GoogleVRHMD/Public/GoogleVRAdbUtils.h
+++ b/Engine/Plugins/Runtime/GoogleVR/GoogleVRHMD/Source/GoogleVRHMD/Public/GoogleVRAdbUtils.h
@@ -10,11 +10,10 @@
 // TODO: would be nice if Unreal make that function public so we don't need to make a duplicate.
 inline void GetAdbPath(FString& OutAdbPath)
 {
-	TCHAR AndroidDirectory[32768] = { 0 };
-	FPlatformMisc::GetEnvironmentVariable(TEXT("ANDROID_HOME"), AndroidDirectory, 32768);
+	FString AndroidDirectory = FPlatformMisc::GetEnvironmentVariable(TEXT("ANDROID_HOME"));
 
 #if PLATFORM_MAC
-	if (AndroidDirectory[0] == 0)
+	if (AndroidDirectory.Len() == 0)
 	{
 		// didn't find ANDROID_HOME, so parse the .bash_profile file on MAC
 		FArchive* FileReader = IFileManager::Get().CreateFileReader(*FString([@"~/.bash_profile" stringByExpandingTildeInPath]));
@@ -33,25 +32,25 @@ inline void GetAdbPath(FString& OutAdbPath)
 
 			for (int32 Index = Lines.Num() - 1; Index >= 0; Index--)
 			{
-				if (AndroidDirectory[0] == 0 && Lines[Index].StartsWith(TEXT("export ANDROID_HOME=")))
+				if (AndroidDirectory.Len() == 0 && Lines[Index].StartsWith(TEXT("export ANDROID_HOME=")))
 				{
 					FString Directory;
 					Lines[Index].Split(TEXT("="), NULL, &Directory);
 					Directory = Directory.Replace(TEXT("\""), TEXT(""));
-					FCString::Strcpy(AndroidDirectory, *Directory);
-					setenv("ANDROID_HOME", TCHAR_TO_ANSI(AndroidDirectory), 1);
+					AndroidDirectory = Directory;
+					setenv("ANDROID_HOME", TCHAR_TO_ANSI(*AndroidDirectory), 1);
 				}
 			}
 		}
 	}
 #endif
 
-	if (AndroidDirectory[0] != 0)
+	if (AndroidDirectory.Len() != 0)
 	{
 #if PLATFORM_WINDOWS
-		OutAdbPath = FString::Printf(TEXT("%s\\platform-tools\\adb.exe"), AndroidDirectory);
+		OutAdbPath = FString::Printf(TEXT("%s\\platform-tools\\adb.exe"), *AndroidDirectory);
 #else
-		OutAdbPath = FString::Printf(TEXT("%s/platform-tools/adb"), AndroidDirectory);
+		OutAdbPath = FString::Printf(TEXT("%s/platform-tools/adb"), *AndroidDirectory);
 #endif
 
 		// if it doesn't exist then just clear the path as we might set it later
diff --git a/Engine/Plugins/Runtime/Steam/SteamVR/Source/SteamVR/Private/SteamVRHMD.cpp b/Engine/Plugins/Runtime/Steam/SteamVR/Source/SteamVR/Private/SteamVRHMD.cpp
--- a/Engine/Plugins/Runtime/Steam/SteamVR/Source/SteamVR/Private/SteamVRHMD.cpp
+++ b/Engine/Plugins/Runtime/Steam/SteamVR/Source/SteamVR/Private/SteamVRHMD.cpp
@@ -156,12 +156,11 @@ public:
 #if PLATFORM_WINDOWS
 #if PLATFORM_64BITS
 		FString RootOpenVRPath;
-		TCHAR VROverridePath[MAX_PATH];
-		FPlatformMisc::GetEnvironmentVariable(TEXT("VR_OVERRIDE"), VROverridePath, MAX_PATH);
+		FString VROverridePath = FPlatformMisc::GetEnvironmentVariable(TEXT("VR_OVERRIDE"));
 		
-		if (FCString::Strlen(VROverridePath) > 0)
+		if (VROverridePath.Len() > 0)
 		{
-			RootOpenVRPath = FString::Printf(TEXT("%s\\bin\\win64\\"), VROverridePath);
+			RootOpenVRPath = FString::Printf(TEXT("%s\\bin\\win64\\"), *VROverridePath);
 		}
 		else
 		{
diff --git a/Engine/Source/Developer/Android/AndroidDeviceDetection/Private/AndroidDeviceDetectionModule.cpp b/Engine/Source/Developer/Android/AndroidDeviceDetection/Private/AndroidDeviceDetectionModule.cpp
--- a/Engine/Source/Developer/Android/AndroidDeviceDetection/Private/AndroidDeviceDetectionModule.cpp
+++ b/Engine/Source/Developer/Android/AndroidDeviceDetection/Private/AndroidDeviceDetectionModule.cpp
@@ -420,13 +420,12 @@ public:
 	virtual void UpdateADBPath() override
 	{
 		FScopeLock PathUpdateLock(&ADBPathCheckLock);
-		TCHAR AndroidDirectory[32768] = { 0 };
-		FPlatformMisc::GetEnvironmentVariable(*SDKDirEnvVar, AndroidDirectory, 32768);
+		FString AndroidDirectory = FPlatformMisc::GetEnvironmentVariable(*SDKDirEnvVar);
 
 		ADBPath.Empty();
 		
 #if PLATFORM_MAC || PLATFORM_LINUX
-		if (AndroidDirectory[0] == 0)
+		if (AndroidDirectory.Len() == 0)
 		{
 #if PLATFORM_LINUX
 			// didn't find ANDROID_HOME, so parse the .bashrc file on Linux
@@ -450,22 +449,22 @@ public:
 
 				for (int32 Index = Lines.Num()-1; Index >=0; Index--)
 				{
-					if (AndroidDirectory[0] == 0 && Lines[Index].StartsWith(FString::Printf(TEXT("export %s="), *SDKDirEnvVar)))
+					if (AndroidDirectory.Len() == 0 && Lines[Index].StartsWith(FString::Printf(TEXT("export %s="), *SDKDirEnvVar)))
 					{
 						FString Directory;
 						Lines[Index].Split(TEXT("="), NULL, &Directory);
 						Directory = Directory.Replace(TEXT("\""), TEXT(""));
-						FCString::Strcpy(AndroidDirectory, *Directory);
-						setenv(TCHAR_TO_ANSI(*SDKDirEnvVar), TCHAR_TO_ANSI(AndroidDirectory), 1);
+						AndroidDirectory = Directory;
+						setenv(TCHAR_TO_ANSI(*SDKDirEnvVar), TCHAR_TO_ANSI(*AndroidDirectory), 1);
 					}
 				}
 			}
 		}
 #endif
 
-		if (AndroidDirectory[0] != 0)
+		if (AndroidDirectory.Len() > 0)
 		{
-			ADBPath = FPaths::Combine(AndroidDirectory, SDKRelativeExePath);
+			ADBPath = FPaths::Combine(*AndroidDirectory, SDKRelativeExePath);
 
 			// if it doesn't exist then just clear the path as we might set it later
 			if (!FPaths::FileExists(*ADBPath))
diff --git a/Engine/Source/Developer/Android/AndroidTargetPlatform/Private/AndroidTargetDevice.inl b/Engine/Source/Developer/Android/AndroidTargetPlatform/Private/AndroidTargetDevice.inl
--- a/Engine/Source/Developer/Android/AndroidTargetPlatform/Private/AndroidTargetDevice.inl
+++ b/Engine/Source/Developer/Android/AndroidTargetPlatform/Private/AndroidTargetDevice.inl
@@ -159,18 +159,16 @@ inline bool FAndroidTargetDevice::GetAdbFullFilename(FString& OutFilename)
 	TOptional<FString> ResultPath;
 
 	// get the SDK binaries folder
-	TCHAR AndroidDirectory[32768] = { 0 };
-
-	FPlatformMisc::GetEnvironmentVariable(TEXT("ANDROID_HOME"), AndroidDirectory, 32768);
-	if (AndroidDirectory[0] == 0)
+	FString AndroidDirectory = FPlatformMisc::GetEnvironmentVariable(TEXT("ANDROID_HOME"));
+	if (AndroidDirectory.Len() == 0)
 	{
 		return false;
 	}
 
 #if PLATFORM_WINDOWS
-	OutFilename = FString::Printf(TEXT("%s\\platform-tools\\adb.exe"), AndroidDirectory);
+	OutFilename = FString::Printf(TEXT("%s\\platform-tools\\adb.exe"), *AndroidDirectory);
 #else
-	OutFilename = FString::Printf(TEXT("%s/platform-tools/adb"), AndroidDirectory);
+	OutFilename = FString::Printf(TEXT("%s/platform-tools/adb"), *AndroidDirectory);
 #endif
 
 	return true;
diff --git a/Engine/Source/Developer/Apple/MetalShaderFormat/Private/MetalDerivedData.cpp b/Engine/Source/Developer/Apple/MetalShaderFormat/Private/MetalDerivedData.cpp
--- a/Engine/Source/Developer/Apple/MetalShaderFormat/Private/MetalDerivedData.cpp
+++ b/Engine/Source/Developer/Apple/MetalShaderFormat/Private/MetalDerivedData.cpp
@@ -166,9 +166,8 @@ bool FMetalShaderBytecodeCooker::Build(TArray<uint8>& OutData)
 
 #if PLATFORM_MAC
 	// Unset the SDKROOT to avoid problems with the incorrect path being used when compiling with the shared PCH.
-	TCHAR SdkRoot[4096];
-	FPlatformMisc::GetEnvironmentVariable(TEXT("SDKROOT"), SdkRoot, ARRAY_COUNT(SdkRoot));
-	if (FCStringWide::Strlen(SdkRoot))
+	FString SdkRoot = FPlatformMisc::GetEnvironmentVariable(TEXT("SDKROOT"));
+	if (SdkRoot.Len() > 0)
 	{
 		unsetenv("SDKROOT");
 	}
@@ -283,9 +282,9 @@ bool FMetalShaderBytecodeCooker::Build(TArray<uint8>& OutData)
 
 #if PLATFORM_MAC
 	// Reset the SDKROOT environment we unset earlier.
-	if (FCStringWide::Strlen(SdkRoot))
+	if (SdkRoot.Len() > 0)
 	{
-		setenv("SDKROOT", TCHAR_TO_UTF8(SdkRoot), 1);
+		setenv("SDKROOT", TCHAR_TO_UTF8(*SdkRoot), 1);
 	}
 #endif
 
diff --git a/Engine/Source/Developer/Apple/MetalShaderFormat/Private/MetalShaderCompiler.cpp b/Engine/Source/Developer/Apple/MetalShaderFormat/Private/MetalShaderCompiler.cpp
--- a/Engine/Source/Developer/Apple/MetalShaderFormat/Private/MetalShaderCompiler.cpp
+++ b/Engine/Source/Developer/Apple/MetalShaderFormat/Private/MetalShaderCompiler.cpp
@@ -163,9 +163,8 @@ bool IsRemoteBuildingConfigured(const FShaderCompilerEnvironment* InEnvironment)
 					if (GRemoteBuildServerSSHKey.Len() == 0)
 					{
 						// RemoteToolChain.cs in UBT looks in a few more places but the code in FIOSTargetSettingsCustomization::OnGenerateSSHKey() only puts the key in this location so just going with that to keep things simple
-						TCHAR Path[4096];
-						FPlatformMisc::GetEnvironmentVariable(TEXT("APPDATA"), Path, ARRAY_COUNT(Path));
-						GRemoteBuildServerSSHKey = FString::Printf(TEXT("%s\\Unreal Engine\\UnrealBuildTool\\SSHKeys\\%s\\%s\\RemoteToolChainPrivate.key"), Path, *GRemoteBuildServerHost, *GRemoteBuildServerUser);
+						FString Path = FPlatformMisc::GetEnvironmentVariable(TEXT("APPDATA"));
+						GRemoteBuildServerSSHKey = FString::Printf(TEXT("%s\\Unreal Engine\\UnrealBuildTool\\SSHKeys\\%s\\%s\\RemoteToolChainPrivate.key"), *Path, *GRemoteBuildServerHost, *GRemoteBuildServerUser);
 					}
 				}
 			}
@@ -204,9 +203,8 @@ bool IsRemoteBuildingConfigured(const FShaderCompilerEnvironment* InEnvironment)
 		if (!FPaths::DirectoryExists(DeltaCopyPath))
 		{
 			// if no UE4 bundled version of DeltaCopy, try and use the default install location
-			TCHAR ProgramPath[4096];
-			FPlatformMisc::GetEnvironmentVariable(TEXT("PROGRAMFILES(X86)"), ProgramPath, ARRAY_COUNT(ProgramPath));
-			DeltaCopyPath = FPaths::Combine(ProgramPath, TEXT("DeltaCopy"));
+			FString ProgramPath = FPlatformMisc::GetEnvironmentVariable(TEXT("PROGRAMFILES(X86)"));
+			DeltaCopyPath = FPaths::Combine(*ProgramPath, TEXT("DeltaCopy"));
 		}
 
 		if (!FPaths::DirectoryExists(DeltaCopyPath))
diff --git a/Engine/Source/Developer/DerivedDataCache/Private/DerivedDataBackends.cpp b/Engine/Source/Developer/DerivedDataCache/Private/DerivedDataBackends.cpp
--- a/Engine/Source/Developer/DerivedDataCache/Private/DerivedDataBackends.cpp
+++ b/Engine/Source/Developer/DerivedDataCache/Private/DerivedDataBackends.cpp
@@ -481,9 +481,8 @@ public:
 		FString EnvPathOverride;
 		if( FParse::Value( Entry, TEXT("EnvPathOverride="), EnvPathOverride ) )
 		{
-			TCHAR FilesystemCachePathEnv[256];
-			FPlatformMisc::GetEnvironmentVariable( *EnvPathOverride, FilesystemCachePathEnv, ARRAY_COUNT(FilesystemCachePathEnv) );
-			if( FilesystemCachePathEnv[0] )
+			FString FilesystemCachePathEnv = FPlatformMisc::GetEnvironmentVariable( *EnvPathOverride );
+			if( FilesystemCachePathEnv.Len() > 0 )
 			{
 				Path = FilesystemCachePathEnv;
 				UE_LOG( LogDerivedDataCache, Log, TEXT("Found environment variable %s=%s"), *EnvPathOverride, *Path );
diff --git a/Engine/Source/Developer/DesktopPlatform/Private/Linux/DesktopPlatformLinux.cpp b/Engine/Source/Developer/DesktopPlatform/Private/Linux/DesktopPlatformLinux.cpp
--- a/Engine/Source/Developer/DesktopPlatform/Private/Linux/DesktopPlatformLinux.cpp
+++ b/Engine/Source/Developer/DesktopPlatform/Private/Linux/DesktopPlatformLinux.cpp
@@ -243,13 +243,10 @@ static bool CompareAndCheckDesktopFile(const TCHAR* DesktopFileName, const TCHAR
 	}
 
 	// There currently appears to be no way to locate the desktop file with xdg-utils so access the file via the expected location.
-	TCHAR DataDir[PLATFORM_MAX_FILEPATH_LENGTH];
-	FPlatformMisc::GetEnvironmentVariable(TEXT("XDG_DATA_HOME"), DataDir, PLATFORM_MAX_FILEPATH_LENGTH);
-	if (FCString::Strlen(DataDir) == 0)
+	FString DataDir = FPlatformMisc::GetEnvironmentVariable(TEXT("XDG_DATA_HOME"));
+	if (DataDir.Len() == 0)
 	{	
-		FPlatformMisc::GetEnvironmentVariable(TEXT("HOME"), DataDir, PLATFORM_MAX_FILEPATH_LENGTH);
-
-		FCString::Strcat(DataDir, TEXT("/.local/share"));
+		DataDir = FPlatformMisc::GetEnvironmentVariable(TEXT("HOME")) + TEXT("/.local/share");
 	}
 
 	// Get the contents of the desktop file.
diff --git a/Engine/Source/Developer/HTML5/HTML5TargetPlatform/Private/HTML5TargetPlatform.cpp b/Engine/Source/Developer/HTML5/HTML5TargetPlatform/Private/HTML5TargetPlatform.cpp
--- a/Engine/Source/Developer/HTML5/HTML5TargetPlatform/Private/HTML5TargetPlatform.cpp
+++ b/Engine/Source/Developer/HTML5/HTML5TargetPlatform/Private/HTML5TargetPlatform.cpp
@@ -90,14 +90,13 @@ bool FHTML5TargetPlatform::IsSdkInstalled(bool bProjectHasCode, FString& OutDocu
 {
 	// When the EMSDK environment variable is used, locate Emscripten SDK from the directory
 	// pointed to by that variable, instead of using a prepackaged SDK.
-	TCHAR EmsdkDirectory[1024] = { 0 };
-	FPlatformMisc::GetEnvironmentVariable(TEXT("EMSDK"), EmsdkDirectory, 1024);
-	if (EmsdkDirectory[0] != 0)
+	FString EmsdkDirectory = FPlatformMisc::GetEnvironmentVariable(TEXT("EMSDK"));
+	if (EmsdkDirectory.Len() > 0)
 	{
-		bool Exists = IFileManager::Get().DirectoryExists(EmsdkDirectory);
+		bool Exists = IFileManager::Get().DirectoryExists(*EmsdkDirectory);
 		if (!Exists)
 		{
-			UE_LOG(LogTemp, Display, TEXT("Environment variable EMSDK is set to \"%s\", but that directory does not exist!"), EmsdkDirectory);
+			UE_LOG(LogTemp, Display, TEXT("Environment variable EMSDK is set to \"%s\", but that directory does not exist!"), *EmsdkDirectory);
 		}
 		return Exists;
 	}
diff --git a/Engine/Source/Developer/IOS/IOSPlatformEditor/Private/IOSTargetSettingsCustomization.cpp b/Engine/Source/Developer/IOS/IOSPlatformEditor/Private/IOSTargetSettingsCustomization.cpp
--- a/Engine/Source/Developer/IOS/IOSPlatformEditor/Private/IOSTargetSettingsCustomization.cpp
+++ b/Engine/Source/Developer/IOS/IOSPlatformEditor/Private/IOSTargetSettingsCustomization.cpp
@@ -1219,15 +1219,15 @@ FReply FIOSTargetSettingsCustomization::OnInstallProvisionClicked()
 
 		// see if the provision is already installed
 		FString DestName = FPaths::GetBaseFilename(ProvisionPath);
-		TCHAR Path[4096];
+		FString Path;
 #if PLATFORM_MAC
-		FPlatformMisc::GetEnvironmentVariable(TEXT("HOME"), Path, ARRAY_COUNT(Path));
-		FString Destination = FString::Printf(TEXT("\"%s/Library/MobileDevice/Provisioning Profiles/%s.mobileprovision\""), Path, *DestName);
-		FString Destination2 = FString::Printf(TEXT("\"%s/Library/MobileDevice/Provisioning Profiles/%s.mobileprovision\""), Path, FApp::GetProjectName());
+		Path = FPlatformMisc::GetEnvironmentVariable(TEXT("HOME"));
+		FString Destination = FString::Printf(TEXT("\"%s/Library/MobileDevice/Provisioning Profiles/%s.mobileprovision\""), *Path, *DestName);
+		FString Destination2 = FString::Printf(TEXT("\"%s/Library/MobileDevice/Provisioning Profiles/%s.mobileprovision\""), *Path, FApp::GetProjectName());
 #else
-		FPlatformMisc::GetEnvironmentVariable(TEXT("LOCALAPPDATA"), Path, ARRAY_COUNT(Path));
-		FString Destination = FString::Printf(TEXT("%s\\Apple Computer\\MobileDevice\\Provisioning Profiles\\%s.mobileprovision"), Path, *DestName);
-		FString Destination2 = FString::Printf(TEXT("%s\\Apple Computer\\MobileDevice\\Provisioning Profiles\\%s.mobileprovision"), Path, FApp::GetProjectName());
+		Path = FPlatformMisc::GetEnvironmentVariable(TEXT("LOCALAPPDATA"));
+		FString Destination = FString::Printf(TEXT("%s\\Apple Computer\\MobileDevice\\Provisioning Profiles\\%s.mobileprovision"), *Path, *DestName);
+		FString Destination2 = FString::Printf(TEXT("%s\\Apple Computer\\MobileDevice\\Provisioning Profiles\\%s.mobileprovision"), *Path, FApp::GetProjectName());
 #endif
 		if (FPaths::FileExists(Destination) || FPaths::FileExists(Destination2))
 		{
@@ -1348,9 +1348,8 @@ FReply FIOSTargetSettingsCustomization::OnGenerateSSHKey()
 		RemoteServerPort = "22";
 	}
 
-	TCHAR Path[4096];
-	FPlatformMisc::GetEnvironmentVariable(TEXT("APPDATA"), Path, ARRAY_COUNT(Path));
-	FString Destination = FString::Printf(TEXT("%s\\Unreal Engine\\UnrealBuildTool\\SSHKeys\\%s\\%s\\RemoteToolChainPrivate.key"), Path, *RemoteServerAddress, *(Settings.RSyncUsername));
+	FString Path = FPlatformMisc::GetEnvironmentVariable(TEXT("APPDATA"));
+	FString Destination = FString::Printf(TEXT("%s\\Unreal Engine\\UnrealBuildTool\\SSHKeys\\%s\\%s\\RemoteToolChainPrivate.key"), *Path, *RemoteServerAddress, *(Settings.RSyncUsername));
 	if (FPaths::FileExists(Destination))
 	{
 		FString MessagePrompt = FString::Printf(TEXT("An SSH Key already exists.  Do you want to replace this key?"));
@@ -1371,9 +1370,8 @@ FReply FIOSTargetSettingsCustomization::OnGenerateSSHKey()
 	if (!FPaths::DirectoryExists(DeltaCopyPath))
 	{
 		// if no UE4 bundled version of DeltaCopy, try and use the default install location
-		TCHAR ProgramPath[4096];
-		FPlatformMisc::GetEnvironmentVariable(TEXT("PROGRAMFILES(X86)"), ProgramPath, ARRAY_COUNT(ProgramPath));
-		DeltaCopyPath = FPaths::Combine(ProgramPath, TEXT("DeltaCopy"));
+		FString ProgramPath = FPlatformMisc::GetEnvironmentVariable(TEXT("PROGRAMFILES(X86)"));
+		DeltaCopyPath = FPaths::Combine(*ProgramPath, TEXT("DeltaCopy"));
 	}
 	
 	if (!FPaths::DirectoryExists(DeltaCopyPath))
diff --git a/Engine/Source/Developer/Linux/LinuxTargetPlatform/Private/LinuxTargetPlatform.h b/Engine/Source/Developer/Linux/LinuxTargetPlatform/Private/LinuxTargetPlatform.h
--- a/Engine/Source/Developer/Linux/LinuxTargetPlatform/Private/LinuxTargetPlatform.h
+++ b/Engine/Source/Developer/Linux/LinuxTargetPlatform/Private/LinuxTargetPlatform.h
@@ -177,19 +177,16 @@ public:
 		if (!PLATFORM_LINUX)
 		{
 			// check for LINUX_MULTIARCH_ROOT or for legacy LINUX_ROOT when targeting Linux from Win/Mac
-			TCHAR ToolchainRoot[32768] = { 0 };
-			FPlatformMisc::GetEnvironmentVariable(TEXT("LINUX_MULTIARCH_ROOT"), ToolchainRoot, ARRAY_COUNT(ToolchainRoot));
+
 			// proceed with any value for MULTIARCH root, because checking exact architecture is not possible at this point
-			FString ToolchainMultiarchRoot = ToolchainRoot;
+			FString ToolchainMultiarchRoot = FPlatformMisc::GetEnvironmentVariable(TEXT("LINUX_MULTIARCH_ROOT"));
 			if (ToolchainMultiarchRoot.Len() > 0 && FPaths::DirectoryExists(ToolchainMultiarchRoot))
 			{
 				return true;
 			}
 			
 			// else check for legacy LINUX_ROOT
-			ToolchainRoot[ 0 ] = 0;
-			FPlatformMisc::GetEnvironmentVariable(TEXT("LINUX_ROOT"), ToolchainRoot, ARRAY_COUNT(ToolchainRoot));
-			FString ToolchainCompiler = ToolchainRoot;
+			FString ToolchainCompiler = FPlatformMisc::GetEnvironmentVariable(TEXT("LINUX_ROOT"));
 			if (PLATFORM_WINDOWS)
 			{
 				ToolchainCompiler += "/bin/clang++.exe";
diff --git a/Engine/Source/Developer/TargetPlatform/Private/TargetPlatformManagerModule.cpp b/Engine/Source/Developer/TargetPlatform/Private/TargetPlatformManagerModule.cpp
--- a/Engine/Source/Developer/TargetPlatform/Private/TargetPlatformManagerModule.cpp
+++ b/Engine/Source/Developer/TargetPlatform/Private/TargetPlatformManagerModule.cpp
@@ -602,12 +602,11 @@ protected:
 	bool IsAutoSDKsEnabled()
 	{
 		static const FString SDKRootEnvFar(TEXT("UE_SDKS_ROOT"));
-		const int32 MaxPathSize = 16384;
-		TCHAR SDKPath[MaxPathSize] = { 0 };
-		FPlatformMisc::GetEnvironmentVariable(*SDKRootEnvFar, SDKPath, MaxPathSize);
+
+		FString SDKPath = FPlatformMisc::GetEnvironmentVariable(*SDKRootEnvFar);
 
 		// AutoSDKs only enabled if UE_SDKS_ROOT is set.
-		if (SDKPath[0] != 0)
+		if (SDKPath.Len() != 0)
 		{
 			return true;
 		}
@@ -737,9 +736,7 @@ RETRY_SETUPANDVALIDATE:
 #endif		
 
 		static const FString SDKRootEnvFar(TEXT("UE_SDKS_ROOT"));
-		const int32 MaxPathSize = 16384;
-		FString SDKPath = FString::ChrN(16384, TEXT('\0'));
-		FPlatformMisc::GetEnvironmentVariable(*SDKRootEnvFar, SDKPath.GetCharArray().GetData(), MaxPathSize);
+		FString SDKPath = FPlatformMisc::GetEnvironmentVariable(*SDKRootEnvFar);
 
 		FString TargetSDKRoot = FPaths::Combine(*SDKPath, *HostPlatform, *AutoSDKPath);
 		static const FString SDKInstallManifestFileName(TEXT("CurrentlyInstalled.txt"));
@@ -842,9 +839,7 @@ RETRY_SETUPANDVALIDATE:
 			SetupEnvironmentVariables(EnvVarNames, EnvVarValues);
 
 
-			const int32 MaxPathVarLen = 32768;
-			FString OrigPathVar = FString::ChrN(MaxPathVarLen, TEXT('\0'));
-			FPlatformMisc::GetEnvironmentVariable(TEXT("PATH"), OrigPathVar.GetCharArray().GetData(), MaxPathVarLen);
+			FString OrigPathVar = FPlatformMisc::GetEnvironmentVariable(TEXT("PATH"));
 
 			// actually perform the PATH stripping / adding.
 			const TCHAR* PathDelimiter = FPlatformMisc::GetPathVarDelimiter();
diff --git a/Engine/Source/Developer/XGEController/Private/XGEControllerModule.cpp b/Engine/Source/Developer/XGEController/Private/XGEControllerModule.cpp
--- a/Engine/Source/Developer/XGEController/Private/XGEControllerModule.cpp
+++ b/Engine/Source/Developer/XGEController/Private/XGEControllerModule.cpp
@@ -200,9 +200,7 @@ bool FXGEControllerModule::IsSupported()
 		// Try to find xgConsole.exe from the PATH environment variable
 		FString PathString;
 		{
-			TCHAR EnvVariable[32 * 1024];
-			FPlatformMisc::GetEnvironmentVariable(TEXT("Path"), EnvVariable, ARRAY_COUNT(EnvVariable));
-			FString EnvString(EnvVariable);
+			FString EnvString = FPlatformMisc::GetEnvironmentVariable(TEXT("Path"));
 			int32 PathStart = EnvString.Find(TEXT("Xoreax\\IncrediBuild"), ESearchCase::IgnoreCase, ESearchDir::FromStart);
 			if (PathStart != INDEX_NONE)
 			{
diff --git a/Engine/Source/Editor/UnrealEd/Private/Commandlets/ContentCommandlets.cpp b/Engine/Source/Editor/UnrealEd/Private/Commandlets/ContentCommandlets.cpp
--- a/Engine/Source/Editor/UnrealEd/Private/Commandlets/ContentCommandlets.cpp
+++ b/Engine/Source/Editor/UnrealEd/Private/Commandlets/ContentCommandlets.cpp
@@ -866,9 +866,8 @@ int32 UResavePackagesCommandlet::Main( const FString& Params )
 	bForceUATEnvironmentVariableSet = false;
 	if (bShouldBuildHLOD)
 	{
-		TCHAR MutexVariableValue = 0;
-		FPlatformMisc::GetEnvironmentVariable(TEXT("uebp_UATMutexNoWait"), &MutexVariableValue, 1);
-		if (MutexVariableValue != 1)
+		FString MutexVariableValue = FPlatformMisc::GetEnvironmentVariable(TEXT("uebp_UATMutexNoWait"));
+		if (MutexVariableValue != TEXT("1"))
 		{
 			FPlatformMisc::SetEnvironmentVar(TEXT("uebp_UATMutexNoWait"), TEXT("1"));
 			bForceUATEnvironmentVariableSet = true;
diff --git a/Engine/Source/Programs/UnrealLightmass/Private/Launch/UnrealLightmass.cpp b/Engine/Source/Programs/UnrealLightmass/Private/Launch/UnrealLightmass.cpp
--- a/Engine/Source/Programs/UnrealLightmass/Private/Launch/UnrealLightmass.cpp
+++ b/Engine/Source/Programs/UnrealLightmass/Private/Launch/UnrealLightmass.cpp
@@ -108,14 +108,12 @@ int LightmassMain(int argc, ANSICHAR* argv[])
 
 	// Override 'NumThreads' with the environment variable, if it's set.
 	{
-		TCHAR* SwarmMaxCoresVariable = new TCHAR[32768];
-		FPlatformMisc::GetEnvironmentVariable( TEXT("Swarm_MaxCores"), SwarmMaxCoresVariable, 32768 );
-		int32 SwarmMaxCores = FCString::Atoi( SwarmMaxCoresVariable );
+		FString SwarmMaxCoresVariable = FPlatformMisc::GetEnvironmentVariable( TEXT("Swarm_MaxCores") );
+		int32 SwarmMaxCores = FCString::Atoi( *SwarmMaxCoresVariable );
 		if ( SwarmMaxCores >= 1 && SwarmMaxCores < 128 )
 		{
 			NumThreads = SwarmMaxCores;
 		}
-		delete [] SwarmMaxCoresVariable;
 	}
 
 	for (int32 ArgIndex = 1; ArgIndex < argc; ArgIndex++)
diff --git a/Engine/Source/Runtime/Core/Private/Android/AndroidMisc.cpp b/Engine/Source/Runtime/Core/Private/Android/AndroidMisc.cpp
--- a/Engine/Source/Runtime/Core/Private/Android/AndroidMisc.cpp
+++ b/Engine/Source/Runtime/Core/Private/Android/AndroidMisc.cpp
@@ -335,6 +335,12 @@ void FAndroidMisc::GetEnvironmentVariable(const TCHAR* VariableName, TCHAR* Resu
 	// @todo Android : get environment variable.
 }
 
+FString FAndroidMisc::GetEnvironmentVariable(const TCHAR* VariableName)
+{
+	// @todo Android : get environment variable.
+	return FString();
+}
+
 const TCHAR* FAndroidMisc::GetSystemErrorMessage(TCHAR* OutBuffer, int32 BufferCount, int32 Error)
 {
 	check(OutBuffer && BufferCount);
diff --git a/Engine/Source/Runtime/Core/Private/Apple/ApplePlatformMisc.cpp b/Engine/Source/Runtime/Core/Private/Apple/ApplePlatformMisc.cpp
--- a/Engine/Source/Runtime/Core/Private/Apple/ApplePlatformMisc.cpp
+++ b/Engine/Source/Runtime/Core/Private/Apple/ApplePlatformMisc.cpp
@@ -30,6 +30,23 @@ void FApplePlatformMisc::GetEnvironmentVariable(const TCHAR* VariableName, TCHAR
 	}
 }
 
+FString FApplePlatformMisc::GetEnvironmentVariable(const TCHAR* VariableName)
+{
+	// Replace hyphens with underscores. Some legacy UE environment variables (eg. UE-SharedDataCachePath) are in widespread
+	// usage in their hyphenated form, but are not normally valid shell variables.
+	FString FixedVariableName = VariableName;
+	FixedVariableName.ReplaceInline(TEXT("-"), TEXT("_"));
+	ANSICHAR *AnsiResult = getenv(TCHAR_TO_ANSI(*FixedVariableName));
+	if (AnsiResult)
+	{
+		return ANSI_TO_TCHAR(AnsiResult);
+	}
+	else
+	{
+		return FString();
+	}
+}
+
 void FApplePlatformMisc::LowLevelOutputDebugString( const TCHAR *Message )
 {
 	//NsLog will out to all iOS output consoles, instead of just the Xcode console.
diff --git a/Engine/Source/Runtime/Core/Private/GenericPlatform/GenericPlatformMisc.cpp b/Engine/Source/Runtime/Core/Private/GenericPlatform/GenericPlatformMisc.cpp
--- a/Engine/Source/Runtime/Core/Private/GenericPlatform/GenericPlatformMisc.cpp
+++ b/Engine/Source/Runtime/Core/Private/GenericPlatform/GenericPlatformMisc.cpp
@@ -222,6 +222,11 @@ FString FSHA256Signature::ToString() const
 	bool FGenericPlatformMisc::bPromptForRemoteDebugOnEnsure = false;
 #endif	//#if !UE_BUILD_SHIPPING
 
+FString FGenericPlatformMisc::GetEnvironmentVariable(const TCHAR* VariableName)
+{
+	return FString();
+}
+
 void FGenericPlatformMisc::SetEnvironmentVar(const TCHAR* VariableName, const TCHAR* Value)
 {
 	UE_LOG(LogGenericPlatformMisc, Error, TEXT("SetEnvironmentVar not implemented for this platform: %s = %s"), VariableName, Value);
diff --git a/Engine/Source/Runtime/Core/Private/Unix/UnixPlatformMisc.cpp b/Engine/Source/Runtime/Core/Private/Unix/UnixPlatformMisc.cpp
--- a/Engine/Source/Runtime/Core/Private/Unix/UnixPlatformMisc.cpp
+++ b/Engine/Source/Runtime/Core/Private/Unix/UnixPlatformMisc.cpp
@@ -196,6 +196,21 @@ void FUnixPlatformMisc::GetEnvironmentVariable(const TCHAR* InVariableName, TCHA
 	}
 }
 
+FString FUnixPlatformMisc::GetEnvironmentVariable(const TCHAR* InVariableName)
+{
+	FString VariableName = InVariableName;
+	VariableName.ReplaceInline(TEXT("-"), TEXT("_"));
+	ANSICHAR *AnsiResult = secure_getenv(TCHAR_TO_ANSI(*VariableName));
+	if (AnsiResult)
+	{
+		return UTF8_TO_TCHAR(AnsiResult);
+	}
+	else
+	{
+		return FString();
+	}
+}
+
 void FUnixPlatformMisc::SetEnvironmentVar(const TCHAR* InVariableName, const TCHAR* Value)
 {
 	FString VariableName = InVariableName;
diff --git a/Engine/Source/Runtime/Core/Private/Windows/WindowsPlatformMisc.cpp b/Engine/Source/Runtime/Core/Private/Windows/WindowsPlatformMisc.cpp
--- a/Engine/Source/Runtime/Core/Private/Windows/WindowsPlatformMisc.cpp
+++ b/Engine/Source/Runtime/Core/Private/Windows/WindowsPlatformMisc.cpp
@@ -614,6 +614,40 @@ void FWindowsPlatformMisc::GetEnvironmentVariable(const TCHAR* VariableName, TCH
 	}
 }
 
+FString FWindowsPlatformMisc::GetEnvironmentVariable(const TCHAR* VariableName)
+{
+	// Attempt to get the environment variable into a local buffer. If it succeeds, we can just copy the result into a string. If it fails, we have the length of the buffer we need to allocate.
+	TCHAR LocalBuffer[128];
+	uint32 Length = ::GetEnvironmentVariableW(VariableName, LocalBuffer, ARRAY_COUNT(LocalBuffer));
+	if (Length < ARRAY_COUNT(LocalBuffer))
+	{
+		LocalBuffer[Length] = 0;
+		return LocalBuffer;
+	}
+
+	// Allocate the data for the string. Loop in case the variable happens to change while running.
+	FString Buffer;
+	for (;;)
+	{
+		TArray<TCHAR>& CharArray = Buffer.GetCharArray();
+		CharArray.SetNumUninitialized(Length);
+
+		Length = ::GetEnvironmentVariableW(VariableName, CharArray.GetData(), CharArray.Num());
+		if (Length == 0)
+		{
+			Buffer.Reset();
+			break;
+		}
+		else if (Length < (uint32)CharArray.Num())
+		{
+			CharArray[Length] = 0;
+			CharArray.SetNum(Length + 1);
+			break;
+		}
+	}
+	return Buffer;
+}
+
 void FWindowsPlatformMisc::SetEnvironmentVar(const TCHAR* VariableName, const TCHAR* Value)
 {
 	uint32 Error = ::SetEnvironmentVariable(VariableName, Value);
diff --git a/Engine/Source/Runtime/Core/Public/Android/AndroidMisc.h b/Engine/Source/Runtime/Core/Public/Android/AndroidMisc.h
--- a/Engine/Source/Runtime/Core/Public/Android/AndroidMisc.h
+++ b/Engine/Source/Runtime/Core/Public/Android/AndroidMisc.h
@@ -32,7 +32,11 @@ struct CORE_API FAndroidMisc : public FGenericPlatformMisc
 	static void PlatformInit();
 	static void PlatformTearDown();
 	static void PlatformHandleSplashScreen(bool ShowSplashScreen);
+
+	DEPRECATED(4.21, "void FPlatformMisc::GetEnvironmentVariable(Name, Result, Length) is deprecated. Use FString FPlatformMisc::GetEnvironmentVariable(Name) instead.")
 	static void GetEnvironmentVariable(const TCHAR* VariableName, TCHAR* Result, int32 ResultLength);
+
+	static FString GetEnvironmentVariable(const TCHAR* VariableName);
 	static const TCHAR* GetSystemErrorMessage(TCHAR* OutBuffer, int32 BufferCount, int32 Error);
 	static EAppReturnType::Type MessageBoxExt( EAppMsgType::Type MsgType, const TCHAR* Text, const TCHAR* Caption );
 	static bool UseRenderThread();
diff --git a/Engine/Source/Runtime/Core/Public/Apple/ApplePlatformMisc.h b/Engine/Source/Runtime/Core/Public/Apple/ApplePlatformMisc.h
--- a/Engine/Source/Runtime/Core/Public/Apple/ApplePlatformMisc.h
+++ b/Engine/Source/Runtime/Core/Public/Apple/ApplePlatformMisc.h
@@ -57,8 +57,12 @@ private:
 struct CORE_API FApplePlatformMisc : public FGenericPlatformMisc
 {
 	static void PlatformInit();
+
+	DEPRECATED(4.21, "void FPlatformMisc::GetEnvironmentVariable(Name, Result, Length) is deprecated. Use FString FPlatformMisc::GetEnvironmentVariable(Name) instead.")
 	static void GetEnvironmentVariable(const TCHAR* VariableName, TCHAR* Result, int32 ResultLength);
 
+	static FString GetEnvironmentVariable(const TCHAR* VariableName);
+
 #if !UE_BUILD_SHIPPING
 	static bool IsDebuggerPresent()
 	{
diff --git a/Engine/Source/Runtime/Core/Public/GenericPlatform/GenericPlatformMisc.h b/Engine/Source/Runtime/Core/Public/GenericPlatform/GenericPlatformMisc.h
--- a/Engine/Source/Runtime/Core/Public/GenericPlatform/GenericPlatformMisc.h
+++ b/Engine/Source/Runtime/Core/Public/GenericPlatform/GenericPlatformMisc.h
@@ -287,11 +287,19 @@ struct CORE_API FGenericPlatformMisc
 	 * @param Result The string to copy the value of the variable into
 	 * @param ResultLength The size of the Result string
 	 */
+	DEPRECATED(4.21, "void FPlatformMisc::GetEnvironmentVariable(Name, Result, Length) is deprecated. Use FString FPlatformMisc::GetEnvironmentVariable(Name) instead.")
 	static void GetEnvironmentVariable(const TCHAR* VariableName, TCHAR* Result, int32 ResultLength)
 	{
 		*Result = 0;
 	}
 
+	/**
+	 * Retrieve a environment variable from the system
+	 *
+	 * @param VariableName The name of the variable (ie "Path")
+	 * @return Value of the variable, or an empty string if not set.
+	 */
+	static FString GetEnvironmentVariable(const TCHAR* VariableName);
 
 	/**
 	 * Sets an environment variable to the local process's environment
diff --git a/Engine/Source/Runtime/Core/Public/Unix/UnixPlatformMisc.h b/Engine/Source/Runtime/Core/Public/Unix/UnixPlatformMisc.h
--- a/Engine/Source/Runtime/Core/Public/Unix/UnixPlatformMisc.h
+++ b/Engine/Source/Runtime/Core/Public/Unix/UnixPlatformMisc.h
@@ -29,7 +29,11 @@ struct CORE_API FUnixPlatformMisc : public FGenericPlatformMisc
 	static void PlatformTearDown();
 	static void SetGracefulTerminationHandler();
 	static void SetCrashHandler(void (* CrashHandler)(const FGenericCrashContext& Context));
+
+	DEPRECATED(4.21, "void FPlatformMisc::GetEnvironmentVariable(Name, Result, Length) is deprecated. Use FString FPlatformMisc::GetEnvironmentVariable(Name) instead.")
 	static void GetEnvironmentVariable(const TCHAR* VariableName, TCHAR* Result, int32 ResultLength);
+
+	static FString GetEnvironmentVariable(const TCHAR* VariableName);
 	static void SetEnvironmentVar(const TCHAR* VariableName, const TCHAR* Value);
 	DEPRECATED(4.14, "GetMacAddress is deprecated. It is not reliable on all platforms")
 	static TArray<uint8> GetMacAddress();
diff --git a/Engine/Source/Runtime/Core/Public/Windows/WindowsPlatformMisc.h b/Engine/Source/Runtime/Core/Public/Windows/WindowsPlatformMisc.h
--- a/Engine/Source/Runtime/Core/Public/Windows/WindowsPlatformMisc.h
+++ b/Engine/Source/Runtime/Core/Public/Windows/WindowsPlatformMisc.h
@@ -43,7 +43,11 @@ struct CORE_API FWindowsPlatformMisc
 	static void PlatformPreInit();
 	static void PlatformInit();
 	static void SetGracefulTerminationHandler();
+
+	DEPRECATED(4.21, "void FPlatformMisc::GetEnvironmentVariable(Name, Result, Length) is deprecated. Use FString FPlatformMisc::GetEnvironmentVariable(Name) instead.")
 	static void GetEnvironmentVariable(const TCHAR* VariableName, TCHAR* Result, int32 ResultLength);
+
+	static FString GetEnvironmentVariable(const TCHAR* VariableName);
 	static void SetEnvironmentVar(const TCHAR* VariableName, const TCHAR* Value);
 
 	static TArray<uint8> GetMacAddress();
diff --git a/Engine/Source/Runtime/Engine/Private/UnrealEngine.cpp b/Engine/Source/Runtime/Engine/Private/UnrealEngine.cpp
--- a/Engine/Source/Runtime/Engine/Private/UnrealEngine.cpp
+++ b/Engine/Source/Runtime/Engine/Private/UnrealEngine.cpp
@@ -8435,8 +8435,7 @@ void UEngine::EnableScreenSaver( bool bEnable )
 		return;
 	}
 
-	TCHAR EnvVariable[32];
-	FPlatformMisc::GetEnvironmentVariable(TEXT("UE-DisallowScreenSaverInhibitor"), EnvVariable, ARRAY_COUNT(EnvVariable));
+	FString EnvVariable = FPlatformMisc::GetEnvironmentVariable(TEXT("UE-DisallowScreenSaverInhibitor"));
 	const bool bDisallowScreenSaverInhibitor = FString(EnvVariable).ToBool();
 
 	// By default we allow to use screen saver inhibitor, but in some cases user can override this setting.
diff --git a/Engine/Source/Runtime/IOS/IOSRuntimeSettings/Private/IOSRuntimeSettings.cpp b/Engine/Source/Runtime/IOS/IOSRuntimeSettings/Private/IOSRuntimeSettings.cpp
--- a/Engine/Source/Runtime/IOS/IOSRuntimeSettings/Private/IOSRuntimeSettings.cpp
+++ b/Engine/Source/Runtime/IOS/IOSRuntimeSettings/Private/IOSRuntimeSettings.cpp
@@ -115,8 +115,8 @@ void UIOSRuntimeSettings::PostInitProperties()
 
 		const FString DefaultKeyFilename = TEXT("RemoteToolChainPrivate.key");
 		const FString RelativeFilePathLocation = FPaths::Combine(TEXT("SSHKeys"), *RemoteServerName, *RSyncUsername, *DefaultKeyFilename);
-		TCHAR Path[4096];
-		FPlatformMisc::GetEnvironmentVariable(TEXT("APPDATA"), Path, ARRAY_COUNT(Path));
+
+		FString Path = FPlatformMisc::GetEnvironmentVariable(TEXT("APPDATA"));
 
 		TArray<FString> PossibleKeyLocations;
 		PossibleKeyLocations.Add(FPaths::Combine(*FPaths::ProjectDir(), TEXT("Build"), TEXT("NotForLicensees"), *RelativeFilePathLocation));
@@ -125,7 +125,7 @@ void UIOSRuntimeSettings::PostInitProperties()
 		PossibleKeyLocations.Add(FPaths::Combine(*FPaths::EngineDir(), TEXT("Build"), TEXT("NotForLicensees"), *RelativeFilePathLocation));
 		PossibleKeyLocations.Add(FPaths::Combine(*FPaths::EngineDir(), TEXT("Build"), TEXT("NoRedist"), *RelativeFilePathLocation));
 		PossibleKeyLocations.Add(FPaths::Combine(*FPaths::EngineDir(), TEXT("Build"), *RelativeFilePathLocation));
-		PossibleKeyLocations.Add(FPaths::Combine(Path, TEXT("Unreal Engine"), TEXT("UnrealBuildTool"), *RelativeFilePathLocation));
+		PossibleKeyLocations.Add(FPaths::Combine(*Path, TEXT("Unreal Engine"), TEXT("UnrealBuildTool"), *RelativeFilePathLocation));
 
 		// Find a potential path that we will use if the user hasn't overridden.
 		// For information purposes only
diff --git a/Engine/Source/Runtime/Launch/Private/LaunchEngineLoop.cpp b/Engine/Source/Runtime/Launch/Private/LaunchEngineLoop.cpp
--- a/Engine/Source/Runtime/Launch/Private/LaunchEngineLoop.cpp
+++ b/Engine/Source/Runtime/Launch/Private/LaunchEngineLoop.cpp
@@ -3919,24 +3919,15 @@ bool FEngineLoop::AppInit( )
 	}
 
 
-	// 8192 is the maximum length of the command line on Windows XP.
-	TCHAR CmdLineEnv[8192];
-
 	// Retrieve additional command line arguments from environment variable.
-	FPlatformMisc::GetEnvironmentVariable(TEXT("UE-CmdLineArgs"), CmdLineEnv,ARRAY_COUNT(CmdLineEnv));
-
-	// Manually nullptr terminate just in case. The nullptr string is returned above in the error case so
-	// we don't have to worry about that.
-	CmdLineEnv[ARRAY_COUNT(CmdLineEnv)-1] = 0;
-	FString Env = FString(CmdLineEnv).TrimStart();
-
+	FString Env = FPlatformMisc::GetEnvironmentVariable(TEXT("UE-CmdLineArgs")).TrimStart();
 	if (Env.Len())
 	{
 		// Append the command line environment after inserting a space as we can't set it in the
 		// environment. Note that any code accessing GCmdLine before appInit obviously won't
 		// respect the command line environment additions.
 		FCommandLine::Append(TEXT(" -EnvAfterHere "));
-		FCommandLine::Append(CmdLineEnv);
+		FCommandLine::Append(*Env);
 	}
 #endif
 
