CIS fix

#rnx

diff --git a/Engine/Source/Developer/IOS/IOSPlatformEditor/Private/IOSTargetSettingsCustomization.cpp b/Engine/Source/Developer/IOS/IOSPlatformEditor/Private/IOSTargetSettingsCustomization.cpp
--- a/Engine/Source/Developer/IOS/IOSPlatformEditor/Private/IOSTargetSettingsCustomization.cpp
+++ b/Engine/Source/Developer/IOS/IOSPlatformEditor/Private/IOSTargetSettingsCustomization.cpp
@@ -1387,7 +1387,7 @@ FReply FIOSTargetSettingsCustomization::OnGenerateSSHKey()
 		*DeltaCopyPath,
 		*(Settings.RSyncUsername),
 		*RemoteServerAddress,
-		Path,
+		*Path,
 		*CygwinPath,
 		*EnginePath);
 
