UE-62145 Only show provisioning warning on iOS builds

#Android

The error message should not appear anymore when opening the Project Settings, only when changing the bundle id or the team name in Project Settings/iOS

diff --git a/Engine/Source/Developer/IOS/IOSPlatformEditor/Private/IOSPlatformEditorModule.cpp b/Engine/Source/Developer/IOS/IOSPlatformEditor/Private/IOSPlatformEditorModule.cpp
--- a/Engine/Source/Developer/IOS/IOSPlatformEditor/Private/IOSPlatformEditorModule.cpp
+++ b/Engine/Source/Developer/IOS/IOSPlatformEditor/Private/IOSPlatformEditorModule.cpp
@@ -1,8 +1,6 @@
 // Copyright 1998-2018 Epic Games, Inc. All Rights Reserved.
 
-#include "CoreMinimal.h"
-#include "Modules/ModuleInterface.h"
-#include "Modules/ModuleManager.h"
+#include "IOSPlatformEditorModule.h"
 #include "Materials/Material.h"
 #include "Materials/MaterialInstance.h"
 #include "PropertyEditorModule.h"
@@ -13,88 +11,81 @@
 #include "MaterialShaderQualitySettingsCustomization.h"
 #include "ComponentRecreateRenderStateContext.h"
 #include "ShaderPlatformQualitySettings.h"
-#include "IOSCustomIconProjectBuildMutatorFeature.h"
 #include "Features/IModularFeatures.h"
+#include "ISettingsSection.h"
 
 #define LOCTEXT_NAMESPACE "FIOSPlatformEditorModule"
 
+FSimpleMulticastDelegate FIOSPlatformEditorModule::OnSelect;
 
-/**
- * Module for iOS as a target platform
- */
-class FIOSPlatformEditorModule
-	: public IModuleInterface
+void FIOSPlatformEditorModule::StartupModule()
 {
-	// IModuleInterface interface
-
-	virtual void StartupModule() override
+	// register settings detail panel customization
+	FPropertyEditorModule& PropertyModule = FModuleManager::LoadModuleChecked<FPropertyEditorModule>("PropertyEditor");
+	PropertyModule.RegisterCustomClassLayout(
+		"IOSRuntimeSettings",
+		FOnGetDetailCustomizationInstance::CreateStatic(&FIOSTargetSettingsCustomization::MakeInstance)
+	);
+
+	FOnUpdateMaterialShaderQuality UpdateMaterials = FOnUpdateMaterialShaderQuality::CreateLambda([]()
 	{
-		// register settings detail panel customization
-		FPropertyEditorModule& PropertyModule = FModuleManager::LoadModuleChecked<FPropertyEditorModule>("PropertyEditor");
-		PropertyModule.RegisterCustomClassLayout(
-			"IOSRuntimeSettings",
-			FOnGetDetailCustomizationInstance::CreateStatic(&FIOSTargetSettingsCustomization::MakeInstance)
+		FGlobalComponentRecreateRenderStateContext Recreate;
+		FlushRenderingCommands();
+		UMaterial::AllMaterialsCacheResourceShadersForRendering();
+		UMaterialInstance::AllMaterialsCacheResourceShadersForRendering();
+	});
+
+	PropertyModule.RegisterCustomClassLayout(
+		UShaderPlatformQualitySettings::StaticClass()->GetFName(),
+		FOnGetDetailCustomizationInstance::CreateStatic(&FMaterialShaderQualitySettingsCustomization::MakeInstance, UpdateMaterials)
 		);
 
-		FOnUpdateMaterialShaderQuality UpdateMaterials = FOnUpdateMaterialShaderQuality::CreateLambda([]()
-		{
-			FGlobalComponentRecreateRenderStateContext Recreate;
-			FlushRenderingCommands();
-			UMaterial::AllMaterialsCacheResourceShadersForRendering();
-			UMaterialInstance::AllMaterialsCacheResourceShadersForRendering();
-		});
-
-		PropertyModule.RegisterCustomClassLayout(
-			UShaderPlatformQualitySettings::StaticClass()->GetFName(),
-			FOnGetDetailCustomizationInstance::CreateStatic(&FMaterialShaderQualitySettingsCustomization::MakeInstance, UpdateMaterials)
-			);
+	PropertyModule.NotifyCustomizationModuleChanged();
+
+	// register settings
+	ISettingsModule* SettingsModule = FModuleManager::GetModulePtr<ISettingsModule>("Settings");
 
-		PropertyModule.NotifyCustomizationModuleChanged();
+	if (SettingsModule != nullptr)
+	{
+		TSharedPtr<ISettingsSection> SelectedSection = SettingsModule->RegisterSettings("Project", "Platforms", "iOS",
+			LOCTEXT("RuntimeSettingsName", "iOS"),
+			LOCTEXT("RuntimeSettingsDescription", "Settings and resources for the iOS platform"),
+			GetMutableDefault<UIOSRuntimeSettings>()
+		);
 
-		// register settings
-		ISettingsModule* SettingsModule = FModuleManager::GetModulePtr<ISettingsModule>("Settings");
+		SelectedSection->OnSelect().BindRaw(this, &FIOSPlatformEditorModule::HandleSelectIOSSection);
 
-		if (SettingsModule != nullptr)
 		{
-			SettingsModule->RegisterSettings("Project", "Platforms", "iOS",
-				LOCTEXT("RuntimeSettingsName", "iOS"),
-				LOCTEXT("RuntimeSettingsDescription", "Settings and resources for the iOS platform"),
-				GetMutableDefault<UIOSRuntimeSettings>()
+			static FName NAME_SF_METAL(TEXT("SF_METAL"));
+			UShaderPlatformQualitySettings* IOSMaterialQualitySettings = UMaterialShaderQualitySettings::Get()->GetShaderPlatformQualitySettings(NAME_SF_METAL);
+			SettingsModule->RegisterSettings("Project", "Platforms", "iOSMetalQuality",
+				LOCTEXT("IOSMetalQualitySettingsName", "iOS Material Quality"),
+				LOCTEXT("IOSMetalQualitySettingsDescription", "Settings for iOS material quality"),
+				IOSMaterialQualitySettings
 			);
-
-			{
-				static FName NAME_SF_METAL(TEXT("SF_METAL"));
-				UShaderPlatformQualitySettings* IOSMaterialQualitySettings = UMaterialShaderQualitySettings::Get()->GetShaderPlatformQualitySettings(NAME_SF_METAL);
-				SettingsModule->RegisterSettings("Project", "Platforms", "iOSMetalQuality",
-					LOCTEXT("IOSMetalQualitySettingsName", "iOS Material Quality"),
-					LOCTEXT("IOSMetalQualitySettingsDescription", "Settings for iOS material quality"),
-					IOSMaterialQualitySettings
-				);
-			}
 		}
-
-		IModularFeatures::Get().RegisterModularFeature(FProjectBuildMutatorFeature::GetFeatureName(), &ProjectBuildMutator);
 	}
 
-	virtual void ShutdownModule() override
-	{
-		ISettingsModule* SettingsModule = FModuleManager::GetModulePtr<ISettingsModule>("Settings");
+	IModularFeatures::Get().RegisterModularFeature(FProjectBuildMutatorFeature::GetFeatureName(), &ProjectBuildMutator);
+}
 
-		if (SettingsModule != nullptr)
-		{
-			SettingsModule->UnregisterSettings("Project", "Platforms", "iOS");
-			SettingsModule->UnregisterSettings("Project", "Platforms", "iOSMetalQuality");
-		}
+void FIOSPlatformEditorModule::ShutdownModule()
+{
+	ISettingsModule* SettingsModule = FModuleManager::GetModulePtr<ISettingsModule>("Settings");
 
-		IModularFeatures::Get().UnregisterModularFeature(FProjectBuildMutatorFeature::GetFeatureName(), &ProjectBuildMutator);
+	if (SettingsModule != nullptr)
+	{
+		SettingsModule->UnregisterSettings("Project", "Platforms", "iOS");
+		SettingsModule->UnregisterSettings("Project", "Platforms", "iOSMetalQuality");
 	}
 
-private:
-
-	FIOSCustomIconProjectBuildMutatorFeature ProjectBuildMutator;
-
-};
+	IModularFeatures::Get().UnregisterModularFeature(FProjectBuildMutatorFeature::GetFeatureName(), &ProjectBuildMutator);
+}
 
+void FIOSPlatformEditorModule::HandleSelectIOSSection()
+{
+	FIOSPlatformEditorModule::OnSelect.Broadcast();
+}
 
 IMPLEMENT_MODULE(FIOSPlatformEditorModule, IOSPlatformEditor);
 
diff --git a/Engine/Source/Developer/IOS/IOSPlatformEditor/Private/IOSPlatformEditorModule.h b/Engine/Source/Developer/IOS/IOSPlatformEditor/Private/IOSPlatformEditorModule.h
new file mode 100644
--- /dev/null
+++ b/Engine/Source/Developer/IOS/IOSPlatformEditor/Private/IOSPlatformEditorModule.h
@@ -0,0 +1,27 @@
+// Copyright 1998-2018 Epic Games, Inc. All Rights Reserved.
+
+#include "IOSRuntimeSettings.h"
+#include "ISettingsModule.h"
+#include "Features/IModularFeatures.h"
+#include "ISettingsSection.h"
+#include "IOSCustomIconProjectBuildMutatorFeature.h"
+
+/**
+ * Module for iOS as a target platform
+ */
+class FIOSPlatformEditorModule : public IModuleInterface
+{
+	// IModuleInterface interface
+public:
+	virtual void StartupModule();
+	virtual void ShutdownModule();
+
+	virtual void HandleSelectIOSSection();
+
+	// Delegate to notify interested parties when the client sources have changed
+	static FSimpleMulticastDelegate OnSelect;
+
+private:
+	FIOSCustomIconProjectBuildMutatorFeature ProjectBuildMutator;
+};
+
diff --git a/Engine/Source/Developer/IOS/IOSPlatformEditor/Private/IOSTargetSettingsCustomization.cpp b/Engine/Source/Developer/IOS/IOSPlatformEditor/Private/IOSTargetSettingsCustomization.cpp
--- a/Engine/Source/Developer/IOS/IOSPlatformEditor/Private/IOSTargetSettingsCustomization.cpp
+++ b/Engine/Source/Developer/IOS/IOSPlatformEditor/Private/IOSTargetSettingsCustomization.cpp
@@ -35,6 +35,7 @@
 #include "Widgets/Notifications/SErrorText.h"
 #include "Interfaces/ITargetPlatform.h"
 #include "Interfaces/ITargetPlatformModule.h"
+#include "IOSPlatformEditorModule.h"
 
 #define LOCTEXT_NAMESPACE "IOSTargetSettings"
 DEFINE_LOG_CATEGORY_STATIC(LogIOSTargetSettings, Log, All);
@@ -108,6 +109,9 @@ FIOSTargetSettingsCustomization::FIOSTargetSettingsCustomization()
 	bShowAllCertificates = false;
 	ProvisionList = MakeShareable(new TArray<ProvisionPtr>());
 	CertificateList = MakeShareable(new TArray<CertificatePtr>());
+
+	FIOSPlatformEditorModule::OnSelect.AddRaw(this, &FIOSTargetSettingsCustomization::OnSelect);
+
 }
 
 FIOSTargetSettingsCustomization::~FIOSTargetSettingsCustomization()
@@ -117,6 +121,8 @@ FIOSTargetSettingsCustomization::~FIOSTargetSettingsCustomization()
 		IPPProcess = NULL;
 		FTicker::GetCoreTicker().RemoveTicker(TickerHandle);
 	}
+
+	FIOSPlatformEditorModule::OnSelect.RemoveAll(this);
 }
 
 void FIOSTargetSettingsCustomization::CustomizeDetails(IDetailLayoutBuilder& DetailLayout)
@@ -130,6 +136,11 @@ void FIOSTargetSettingsCustomization::CustomizeDetails(IDetailLayoutBuilder& Det
 	BuildRemoteBuildingSection(DetailLayout);
 
 	AudioPluginWidgetManager.BuildAudioCategory(DetailLayout, EAudioPlatform::IOS);
+
+}
+
+void FIOSTargetSettingsCustomization::OnSelect()
+{
 	FindRequiredFiles();
 }
 
diff --git a/Engine/Source/Developer/IOS/IOSPlatformEditor/Private/IOSTargetSettingsCustomization.h b/Engine/Source/Developer/IOS/IOSPlatformEditor/Private/IOSTargetSettingsCustomization.h
--- a/Engine/Source/Developer/IOS/IOSPlatformEditor/Private/IOSTargetSettingsCustomization.h
+++ b/Engine/Source/Developer/IOS/IOSPlatformEditor/Private/IOSTargetSettingsCustomization.h
@@ -228,6 +228,9 @@ private:
 	//
 	FText GetIOSTeamIDText(TSharedRef<IPropertyHandle> InPropertyHandle) const;
 
+	/** Module callback for when the iOS settings tab is selected */
+	void OnSelect();
+
 	TSharedPtr< SEditableTextBox > BundleIdTextBox;
 
 	TSharedPtr< SEditableTextBox > IOSTeamIDTextBox;
diff --git a/Engine/Source/Developer/Settings/Private/SettingsSection.cpp b/Engine/Source/Developer/Settings/Private/SettingsSection.cpp
--- a/Engine/Source/Developer/Settings/Private/SettingsSection.cpp
+++ b/Engine/Source/Developer/Settings/Private/SettingsSection.cpp
@@ -253,3 +253,11 @@ bool FSettingsSection::SaveDefaults()
 
 	return false;
 }
+
+void FSettingsSection::Select() 
+{
+	if (SelectDelegate.IsBound())
+	{
+		SelectDelegate.Execute();
+	}
+}
\ No newline at end of file
diff --git a/Engine/Source/Developer/Settings/Private/SettingsSection.h b/Engine/Source/Developer/Settings/Private/SettingsSection.h
--- a/Engine/Source/Developer/Settings/Private/SettingsSection.h
+++ b/Engine/Source/Developer/Settings/Private/SettingsSection.h
@@ -100,11 +100,18 @@ public:
 		return StatusDelegate;
 	}
 
+	virtual FOnSelect& OnSelect() override
+	{
+		return SelectDelegate;
+	}
+
 	virtual bool ResetDefaults() override;
 	virtual bool NotifySectionOnPropertyModified() override;
 	virtual bool Save() override;
 	virtual bool SaveDefaults() override;
 
+	virtual void Select() override;
+
 private:
 
 	/** Holds a pointer to the settings category that owns this section. */
@@ -150,4 +157,7 @@ private:
 
 	/** Holds a delegate that is executed to retrieve a status message for a settings section. */
 	FOnStatus StatusDelegate;
+
+	/** Holds a delegate that is executed when a settings section has been selected. */
+	FOnSelect SelectDelegate;
 };
diff --git a/Engine/Source/Developer/Settings/Public/ISettingsSection.h b/Engine/Source/Developer/Settings/Public/ISettingsSection.h
--- a/Engine/Source/Developer/Settings/Public/ISettingsSection.h
+++ b/Engine/Source/Developer/Settings/Public/ISettingsSection.h
@@ -164,6 +164,11 @@ public:
 	 */
 	virtual bool SaveDefaults() = 0;
 
+	/**
+	* Called when the current section is selected.
+	*/
+	virtual void Select() = 0;
+
 public:
 
 	/**
@@ -232,6 +237,13 @@ public:
 	DECLARE_DELEGATE_RetVal(FText, FOnStatus)
 	virtual FOnStatus& OnStatus() = 0;
 
+
+	/**
+	* A delegate that is executed to check whether a settings section can be edited.
+	*
+	*/
+	DECLARE_DELEGATE(FOnSelect)
+	virtual FOnSelect& OnSelect() = 0;
 public:
 
 	/** Virtual destructor. */
diff --git a/Engine/Source/Developer/SettingsEditor/Private/Widgets/SSettingsSectionHeader.cpp b/Engine/Source/Developer/SettingsEditor/Private/Widgets/SSettingsSectionHeader.cpp
--- a/Engine/Source/Developer/SettingsEditor/Private/Widgets/SSettingsSectionHeader.cpp
+++ b/Engine/Source/Developer/SettingsEditor/Private/Widgets/SSettingsSectionHeader.cpp
@@ -530,6 +530,7 @@ void FSettingsDetailRootObjectCustomization::OnSelectedSectionChanged()
 	if(SelectedSection.IsValid())
 	{
 		SelectedSectionObject = SelectedSection->GetSettingsObject();
+		SelectedSection->Select();
 	}
 	else
 	{
