Error  when building the UE4 Editor - adding pragma once in IOSPlatformEditorModule.h

#4.22

diff --git a/Engine/Source/Developer/IOS/IOSPlatformEditor/Private/IOSPlatformEditorModule.h b/Engine/Source/Developer/IOS/IOSPlatformEditor/Private/IOSPlatformEditorModule.h
--- a/Engine/Source/Developer/IOS/IOSPlatformEditor/Private/IOSPlatformEditorModule.h
+++ b/Engine/Source/Developer/IOS/IOSPlatformEditor/Private/IOSPlatformEditorModule.h
@@ -1,9 +1,10 @@
 // Copyright 1998-2018 Epic Games, Inc. All Rights Reserved.
 
-#include "IOSRuntimeSettings.h"
-#include "ISettingsModule.h"
-#include "Features/IModularFeatures.h"
-#include "ISettingsSection.h"
+#pragma once
+
+
+#include "CoreMinimal.h"
+#include "Modules/ModuleInterface.h"
 #include "IOSCustomIconProjectBuildMutatorFeature.h"
 
 /**
