UE-63767 - iOS Device Output Log console does not work

diff --git a/Engine/Source/Developer/IOS/IOSTargetPlatform/Private/IOSTargetDevice.cpp b/Engine/Source/Developer/IOS/IOSTargetPlatform/Private/IOSTargetDevice.cpp
--- a/Engine/Source/Developer/IOS/IOSTargetPlatform/Private/IOSTargetDevice.cpp
+++ b/Engine/Source/Developer/IOS/IOSTargetPlatform/Private/IOSTargetDevice.cpp
@@ -189,6 +189,23 @@ bool FIOSTargetDevice::GetUserCredentials(FString& OutUserName, FString& OutUser
 {
 	return false;
 }
+
+inline void FIOSTargetDevice::ExecuteConsoleCommand(const FString& ExecCommand) const
+{
+	FString DSFilename = FPaths::ConvertRelativePathToFull(FPaths::EngineDir() / TEXT("Binaries/DotNET/IOS/DeploymentServer.exe"));
+	FString Params = FString::Printf(TEXT(" command -device %s -param \"%s\""), *DeviceId.GetDeviceName(), *ExecCommand);
+	FString WorkingFoder = FPaths::ConvertRelativePathToFull(FPaths::EngineDir() / TEXT("Binaries/DotNET/IOS/"));
+
+#if PLATFORM_MAC
+	// On Mac we launch UBT with Mono
+	FString ScriptPath = FPaths::ConvertRelativePathToFull(FPaths::EngineDir() / TEXT("Build/BatchFiles/Mac/RunMono.sh"));
+	Params = FString::Printf(TEXT("\"%s\" \"%s\" %s"), *ScriptPath, *DSFilename, *Params);
+	DSFilename = TEXT("/bin/sh");
+#endif
+
+	FPlatformProcess::CreateProc(*DSFilename, *Params, true, false, false, NULL, 0, *WorkingFoder, NULL);
+}
+
 inline ITargetDeviceOutputPtr FIOSTargetDevice::CreateDeviceOutputRouter(FOutputDevice* Output) const
 {
 	FIOSTargetDeviceOutputPtr DeviceOutputPtr = MakeShareable(new FIOSTargetDeviceOutput());
diff --git a/Engine/Source/Developer/IOS/IOSTargetPlatform/Private/IOSTargetDevice.h b/Engine/Source/Developer/IOS/IOSTargetPlatform/Private/IOSTargetDevice.h
--- a/Engine/Source/Developer/IOS/IOSTargetPlatform/Private/IOSTargetDevice.h
+++ b/Engine/Source/Developer/IOS/IOSTargetPlatform/Private/IOSTargetDevice.h
@@ -59,6 +59,7 @@ public:
 	virtual bool TerminateProcess(const int64 ProcessId) override;
 	virtual void SetUserCredentials(const FString& UserName, const FString& UserPassword) override;
 	virtual bool GetUserCredentials(FString& OutUserName, FString& OutUserPassword) override;
+	virtual void ExecuteConsoleCommand(const FString& ExecCommand) const override;
 	virtual ITargetDeviceOutputPtr CreateDeviceOutputRouter(FOutputDevice* Output) const override;
 
 public:
diff --git a/Engine/Source/Runtime/Launch/Private/IOS/IOSAppDelegateConsoleHandling.cpp b/Engine/Source/Runtime/Launch/Private/IOS/IOSAppDelegateConsoleHandling.cpp
--- a/Engine/Source/Runtime/Launch/Private/IOS/IOSAppDelegateConsoleHandling.cpp
+++ b/Engine/Source/Runtime/Launch/Private/IOS/IOSAppDelegateConsoleHandling.cpp
@@ -298,5 +298,14 @@ extern bool GShowSplashScreen;
 
 #endif // !UE_BUILD_SHIPPING && !TVOS
 
+void EnqueueConsoleCommand(uint8 *Command)
+{
+	if (!Command)
+	{
+		return;
+	}
+
+	[[IOSAppDelegate GetDelegate] HandleConsoleCommand:[NSString stringWithUTF8String : (const char *)Command]];
+}
 
 @end
diff --git a/Engine/Source/Runtime/Launch/Private/IOS/LaunchIOS.cpp b/Engine/Source/Runtime/Launch/Private/IOS/LaunchIOS.cpp
--- a/Engine/Source/Runtime/Launch/Private/IOS/LaunchIOS.cpp
+++ b/Engine/Source/Runtime/Launch/Private/IOS/LaunchIOS.cpp
@@ -21,6 +21,9 @@
 #include "Misc/ConfigCacheIni.h"
 #include "MoviePlayer.h"
 #include "PreLoadScreenManager.h"
+#include "TcpConsoleListener.h"
+#include "Interfaces/IPv4/IPv4Address.h"
+#include "Interfaces/IPv4/IPv4Endpoint.h"
 
 FEngineLoop GEngineLoop;
 FGameLaunchDaemonMessageHandler GCommandSystem;
@@ -275,6 +278,8 @@ void FAppEntry::PlatformInit()
 	IConsoleManager::Get().CallAllConsoleVariableSinks();
 }
 
+extern TcpConsoleListener *ConsoleListener;
+
 void FAppEntry::Init()
 {
 	FPlatformProcess::SetRealTimeMode();
@@ -319,6 +324,13 @@ void FAppEntry::Init()
 
 	// start up the engine
 	GEngineLoop.Init();
+#ifndef UE_BUILD_SHIPPING
+	if (ConsoleListener == nullptr)
+	{
+		FIPv4Endpoint ConsoleTCP(FIPv4Address::InternalLoopback, 8888); //TODO: @csulea read this from some .ini
+		ConsoleListener = new TcpConsoleListener(ConsoleTCP);
+	}
+#endif // UE_BUILD_SHIPPING
 }
 
 static FSuspendRenderingThread* SuspendThread = NULL;
@@ -348,6 +360,10 @@ void FAppEntry::SuspendTick()
 
 void FAppEntry::Shutdown()
 {
+	if (ConsoleListener)
+	{
+		delete ConsoleListener;
+	}
 	NSLog(@"%s", "Shutting down Game ULD Communications\n");
 	GCommandSystem.Shutdown();
     
diff --git a/Engine/Source/Runtime/Networking/Public/Common/TcpSocketBuilder.h b/Engine/Source/Runtime/Networking/Public/Common/TcpSocketBuilder.h
--- a/Engine/Source/Runtime/Networking/Public/Common/TcpSocketBuilder.h
+++ b/Engine/Source/Runtime/Networking/Public/Common/TcpSocketBuilder.h
@@ -215,7 +215,7 @@ public:
 
 		if (SocketSubsystem != nullptr)
 		{
-			Socket = SocketSubsystem->CreateSocket(NAME_Stream, *Description, true);
+			Socket = SocketSubsystem->CreateSocket(NAME_Stream, *Description, ESocketProtocolFamily::IPv4);
 
 			if (Socket != nullptr)
 			{
@@ -225,7 +225,7 @@ public:
 
 				if (!Error)
 				{
-					Error = Bound && !Socket->Bind(*BoundEndpoint.ToInternetAddr());
+					Error = Bound && !Socket->Bind(*BoundEndpoint.ToInternetAddrIPV4());
 				}
 
 				if (!Error)
@@ -255,7 +255,7 @@ public:
 
 				if (Error)
 				{
-					GLog->Logf(TEXT("FTcpSocketBuilder: Failed to create the socket %s as configured"), *Description);
+					GLog->Logf(TEXT("FTcpSocketBuilder: Failed to create the socket %s as configured (%i)"), *Description, errno);
 
 					SocketSubsystem->DestroySocket(Socket);
 
diff --git a/Engine/Source/Runtime/Networking/Public/Interfaces/IPv4/IPv4Endpoint.h b/Engine/Source/Runtime/Networking/Public/Interfaces/IPv4/IPv4Endpoint.h
--- a/Engine/Source/Runtime/Networking/Public/Interfaces/IPv4/IPv4Endpoint.h
+++ b/Engine/Source/Runtime/Networking/Public/Interfaces/IPv4/IPv4Endpoint.h
@@ -113,6 +113,31 @@ public:
 		return InternetAddr;
 	}
 
+	/**
+	 * Converts this endpoint to an FInternetAddr object.
+	 *
+	 * Note: this method will be removed after the socket subsystem is refactored.
+	 *
+	 * @return Internet address object representing this endpoint.
+	 */
+	TSharedRef<FInternetAddr> ToInternetAddrIPV4() const
+	{
+		check(CachedSocketSubsystem != nullptr && "Networking module not loaded and initialized");
+
+		TSharedRef<FInternetAddr> InternetAddr = CachedSocketSubsystem->CreateInternetAddr();
+		{
+			TArray<uint8> RawAddress;
+			RawAddress.Add(Address.A);
+			RawAddress.Add(Address.B);
+			RawAddress.Add(Address.C);
+			RawAddress.Add(Address.D);
+			InternetAddr->SetRawIp(RawAddress);
+			InternetAddr->SetPort(Port);
+		}
+
+		return InternetAddr;
+	}
+
 	/**
 	 * Gets a string representation for this endpoint.
 	 *
diff --git a/Engine/Source/Runtime/Sockets/Private/IOS/SocketSubsystemIOS.cpp b/Engine/Source/Runtime/Sockets/Private/IOS/SocketSubsystemIOS.cpp
--- a/Engine/Source/Runtime/Sockets/Private/IOS/SocketSubsystemIOS.cpp
+++ b/Engine/Source/Runtime/Sockets/Private/IOS/SocketSubsystemIOS.cpp
@@ -101,7 +101,10 @@ FSocket* FSocketSubsystemIOS::CreateSocket(const FName& SocketType, const FStrin
 	FSocketBSD* NewSocket = (FSocketBSD*)FSocketSubsystemBSD::CreateSocket(SocketType, SocketDescription, ProtocolType);
 	if (NewSocket)
 	{
-		NewSocket->SetIPv6Only(false);
+		if (ProtocolType != ESocketProtocolFamily::IPv4)
+		{
+			NewSocket->SetIPv6Only(false);
+		}
 
 		// disable the SIGPIPE exception 
 		int bAllow = 1;
