Unshelved from pending changelist '4831508':

UE-69143 - //UE4/Main 'Build DCC Win64' - Warning: CreateProc failed: The system cannot find the file specified

diff --git a/Engine/Source/Developer/IOS/IOSTargetPlatform/Private/Windows/IOSDeviceHelperWindows.cpp b/Engine/Source/Developer/IOS/IOSTargetPlatform/Private/Windows/IOSDeviceHelperWindows.cpp
--- a/Engine/Source/Developer/IOS/IOSTargetPlatform/Private/Windows/IOSDeviceHelperWindows.cpp
+++ b/Engine/Source/Developer/IOS/IOSTargetPlatform/Private/Windows/IOSDeviceHelperWindows.cpp
@@ -175,8 +175,13 @@ public:
 		{
 			if (bCheckDevices)
 			{
-				// BHP - Turning off device check to prevent it from interfering with packaging
-				QueryDevices();
+#if WITH_EDITOR
+				if (GIsEditor || FPlatformProperties::IsProgram())
+				{
+					// BHP - Turning off device check to prevent it from interfering with packaging
+					QueryDevices();
+				}
+#endif
 			}
 
 			FPlatformProcess::Sleep(5.0f);
