UE-69143 - //UE4/Main 'Build DCC Win64' - Warning: CreateProc failed: The system cannot find the file specified

diff --git a/Engine/Source/Developer/IOS/IOSTargetPlatform/Private/Windows/IOSDeviceHelperWindows.cpp b/Engine/Source/Developer/IOS/IOSTargetPlatform/Private/Windows/IOSDeviceHelperWindows.cpp
--- a/Engine/Source/Developer/IOS/IOSTargetPlatform/Private/Windows/IOSDeviceHelperWindows.cpp
+++ b/Engine/Source/Developer/IOS/IOSTargetPlatform/Private/Windows/IOSDeviceHelperWindows.cpp
@@ -209,10 +209,11 @@ public:
 	}
 
 private:
+	
 	bool ExecuteDSCommand(const FString& CommandLine, FString* OutStdOut, FString* OutStdErr) const
 	{
-		FString ExecutablePath = FString::Printf(TEXT("%sBinaries/DotNET/IOS"), *FPaths::EngineDir());
-		FString Filename = TEXT("DeploymentServer.exe");
+		FString DSFilename = FPaths::ConvertRelativePathToFull(FPaths::EngineDir() / TEXT("Binaries/DotNET/IOS/DeploymentServer.exe"));
+		FString WorkingFoder = FPaths::ConvertRelativePathToFull(FPaths::EngineDir() / TEXT("Binaries/DotNET/IOS/"));
 
 		// execute the command
 		int32 ReturnCode;
@@ -227,7 +228,7 @@ private:
 		void* WritePipe;
 		void* ReadPipe;
 		FPlatformProcess::CreatePipe(ReadPipe, WritePipe);
-		FProcHandle ProcessHandle = FPlatformProcess::CreateProc(*(ExecutablePath / Filename), *CommandLine, false, true, true, NULL, 0, *ExecutablePath, WritePipe);
+		FProcHandle ProcessHandle = FPlatformProcess::CreateProc(*DSFilename, *CommandLine, false, true, true, NULL, 0, *WorkingFoder, WritePipe);
 		while (FPlatformProcess::IsProcRunning(ProcessHandle))
 		{
 			FString NewLine = FPlatformProcess::ReadPipe(ReadPipe);
