UE-69143 - //UE4/Main 'Build DCC Win64' - Warning: CreateProc failed: The system cannot find the file specified #3

diff --git a/Engine/Source/Developer/IOS/IOSTargetPlatform/Private/Windows/IOSDeviceHelperWindows.cpp b/Engine/Source/Developer/IOS/IOSTargetPlatform/Private/Windows/IOSDeviceHelperWindows.cpp
--- a/Engine/Source/Developer/IOS/IOSTargetPlatform/Private/Windows/IOSDeviceHelperWindows.cpp
+++ b/Engine/Source/Developer/IOS/IOSTargetPlatform/Private/Windows/IOSDeviceHelperWindows.cpp
@@ -176,7 +176,7 @@ public:
 			if (bCheckDevices)
 			{
 #if WITH_EDITOR
-				if (GIsEditor || FPlatformProperties::IsProgram())
+				if (!IsRunningCommandlet())
 				{
 					// BHP - Turning off device check to prevent it from interfering with packaging
 					QueryDevices();
