Detect TVOS or IOS connected device type so Launch on Windows uses the correct target
Rebuild DeploymentServer

diff --git a/Engine/Source/Developer/IOS/IOSTargetPlatform/Private/Windows/IOSDeviceHelperWindows.cpp b/Engine/Source/Developer/IOS/IOSTargetPlatform/Private/Windows/IOSDeviceHelperWindows.cpp
--- a/Engine/Source/Developer/IOS/IOSTargetPlatform/Private/Windows/IOSDeviceHelperWindows.cpp
+++ b/Engine/Source/Developer/IOS/IOSTargetPlatform/Private/Windows/IOSDeviceHelperWindows.cpp
@@ -9,6 +9,7 @@ struct FDeviceNotificationCallbackInformation
 {
 	FString UDID;
 	FString	DeviceName;
+	FString ProductType;
 	uint32 msgType;
 };
 
@@ -290,9 +291,10 @@ private:
 			}
 
 			// grab the device serial number
+			int32 typeIndex = DeviceString.Find(TEXT("TYPE: "));
 			int32 idIndex = DeviceString.Find(TEXT("ID: "));
 			int32 nameIndex = DeviceString.Find(TEXT("NAME: "));
-			if (idIndex < 0 || nameIndex < 0)
+			if (typeIndex < 0 || idIndex < 0 || nameIndex < 0)
 			{
 				continue;
 			}
@@ -307,13 +309,15 @@ private:
 				continue;
 			}
 
-			// parse device name
+			// parse product type and device name
+			FString ProductType = DeviceString.Mid(typeIndex + 6, idIndex - 1 - (typeIndex + 6));
 			FString DeviceName = DeviceString.Mid(nameIndex + 6, DeviceString.Len() - (nameIndex + 6));
 
 			// create an FIOSDevice
 			FDeviceNotificationCallbackInformation CallbackInfo;
 			CallbackInfo.DeviceName = DeviceName;
 			CallbackInfo.UDID = SerialNumber;
+			CallbackInfo.ProductType = ProductType;
 			CallbackInfo.msgType = 1;
 			DeviceNotification.Broadcast(&CallbackInfo);
 		}
@@ -417,8 +421,9 @@ void FIOSDeviceHelper::DoDeviceConnect(void* CallbackInfo)
 
 	// fire the event
 	FIOSLaunchDaemonPong Event;
-	Event.DeviceID = FString::Printf(TEXT("IOS@%s"), *(cbi->UDID));
+	Event.DeviceID = FString::Printf(TEXT("%s@%s"), cbi->ProductType.Contains(TEXT("AppleTV")) ? TEXT("TVOS") : TEXT("IOS"), *(cbi->UDID));
 	Event.DeviceName = cbi->DeviceName;
+	Event.DeviceType = cbi->ProductType;
 	Event.bCanReboot = false;
 	Event.bCanPowerOn = false;
 	Event.bCanPowerOff = false;
diff --git a/Engine/Source/Programs/IOS/DeploymentServer/DeployTime.cs b/Engine/Source/Programs/IOS/DeploymentServer/DeployTime.cs
--- a/Engine/Source/Programs/IOS/DeploymentServer/DeployTime.cs
+++ b/Engine/Source/Programs/IOS/DeploymentServer/DeployTime.cs
@@ -353,9 +353,10 @@ namespace DeploymentServer
         {
             PerformActionOnAllDevices(2 * StandardEnumerationDelayMS, delegate(MobileDeviceInstance Device)
             {
-                string DeviceName = Device.DeviceName;
+				string ProductType = Device.ProductType;
+				string DeviceName = Device.DeviceName;
                 string UDID = Device.DeviceId;
-				ReportIF.Log(String.Format("FOUND: ID: {0} NAME: {1}", UDID, DeviceName));
+				ReportIF.Log(String.Format("FOUND: TYPE: {0} ID: {1} NAME: {2}", ProductType, UDID, DeviceName));
 
                 return true;
             });
