UE-63767 - iOS Device Output Log console does not work

diff --git a/Engine/Source/Developer/IOS/IOSTargetPlatform/IOSTargetPlatform.Build.cs b/Engine/Source/Developer/IOS/IOSTargetPlatform/IOSTargetPlatform.Build.cs
--- a/Engine/Source/Developer/IOS/IOSTargetPlatform/IOSTargetPlatform.Build.cs
+++ b/Engine/Source/Developer/IOS/IOSTargetPlatform/IOSTargetPlatform.Build.cs
@@ -17,7 +17,10 @@ public class IOSTargetPlatform : ModuleRules
 				"DesktopPlatform",
 				"LaunchDaemonMessages",
 				"Projects",
-				"AudioPlatformConfiguration"
+				"AudioPlatformConfiguration",
+				"Sockets",
+				"Networking",
+				"Engine"
 			}
 		);
 
@@ -36,6 +39,7 @@ public class IOSTargetPlatform : ModuleRules
 			{
 				"Runtime/Core/Public/Apple",
 				"Runtime/Core/Public/IOS",
+				"Runtime/Networking/Public",
 			}
 		);
 
diff --git a/Engine/Source/Developer/IOS/IOSTargetPlatform/Private/IOSTargetDevice.cpp b/Engine/Source/Developer/IOS/IOSTargetPlatform/Private/IOSTargetDevice.cpp
--- a/Engine/Source/Developer/IOS/IOSTargetPlatform/Private/IOSTargetDevice.cpp
+++ b/Engine/Source/Developer/IOS/IOSTargetPlatform/Private/IOSTargetDevice.cpp
@@ -8,7 +8,220 @@
 #include "MessageEndpoint.h"
 #include "MessageEndpointBuilder.h"
 #include "IOSTargetDeviceOutput.h"
+#if PLATFORM_WINDOWS
+// for windows mutex
+#include "Windows/AllowWindowsPlatformTypes.h"
+#endif // #if PLATFORM_WINDOWS
 
+enum
+{
+    DEFAULT_DS_COMMANDER_PORT = 41000, // default port to use when issuing DeploymentServer commands
+};
+
+FTcpDSCommander::FTcpDSCommander(const uint8* Data, int32 Count, void* WPipe)
+: bStopping(false)
+, bStoped(true)
+, bIsSuccess(false)
+, DSSocket(nullptr)
+, Thread(nullptr)
+, WritePipe(WPipe)
+, DSCommand(nullptr)
+, DSCommandLen(Count + 1)
+, LastActivity(0.0)
+{
+    if (Count > 0)
+    {
+        DSCommand = (uint8*)FMemory::Malloc(Count + 1);
+        FPlatformMemory::Memcpy(DSCommand, Data, Count);
+        DSCommand[Count] = '\n';
+        Thread = FRunnableThread::Create(this, TEXT("FTcpDSCommander"), 128 * 1024, TPri_Normal);
+    }
+}
+
+/** Virtual destructor. */
+FTcpDSCommander::~FTcpDSCommander()
+{
+    if (Thread != nullptr)
+    {
+        Thread->Kill(true);
+        delete Thread;
+    }
+    if (DSCommand)
+    {
+        FMemory::Free(DSCommand);
+    }
+}
+
+void FTcpDSCommander::Exit()
+{
+    // do nothing
+    bStoped = true;
+}
+
+bool FTcpDSCommander::Init()
+{
+    if (DSCommandLen < 1)
+    {
+        bIsSuccess = true;
+        return true;
+    }
+    // create the socket
+    ISocketSubsystem* SSS = ISocketSubsystem::Get();
+    DSSocket = SSS->CreateSocket(NAME_Stream, TEXT("DSCommander tcp"));
+    if (DSSocket == nullptr)
+    {
+        return false;
+    }
+    TSharedRef<FInternetAddr> Addr = SSS->CreateInternetAddr(0, DEFAULT_DS_COMMANDER_PORT);
+    bool bIsValid;
+    Addr->SetIp(TEXT("127.0.0.1"), bIsValid);
+    
+    
+    // try to connect to the server
+    if (DSSocket->Connect(*Addr) == false)
+    {
+        StartDSProcess();
+        if (DSSocket->Connect(*Addr) == false)
+        {
+            // on failure, shut it all down
+            SSS->DestroySocket(DSSocket);
+            DSSocket = nullptr;
+            UE_LOG(LogTemp, Error, TEXT("Failed to connect to deployment server at %s."), *Addr->ToString(true));
+            return false;
+        }
+    }
+    LastActivity = FPlatformTime::Seconds();
+    
+    return true;
+}
+
+uint32 FTcpDSCommander::Run()
+{
+    if (!DSSocket)
+    {
+        UE_LOG(LogTemp, Error, TEXT("Socket not created."));
+        return 1;
+    }
+    int32 NSent = 0;
+    bool BSent = DSSocket->Send(DSCommand, DSCommandLen, NSent);
+    if (NSent != DSCommandLen || !BSent)
+    {
+        Stop();
+        UE_LOG(LogTemp, Error, TEXT("Socket send error."));
+        return 1;
+    }
+    bStoped = false;
+    
+    static const SIZE_T CommandSize = 1024;
+    uint8 RecvBuffer[CommandSize];
+    
+    while (!bStopping)
+    {
+        uint32 Pending = 0;
+        if (DSSocket->GetConnectionState() != ESocketConnectionState::SCS_Connected)
+        {
+            Stop();
+            UE_LOG(LogTemp, Error, TEXT("Socket connection error."));
+            return 1;
+        }
+        if (DSSocket->HasPendingData(Pending))
+        {
+            NSent = 0;
+            FMemory::Memset(RecvBuffer, 0, CommandSize);
+            LastActivity = FPlatformTime::Seconds();
+            if (DSSocket->Recv(RecvBuffer, CommandSize, NSent))
+            {
+                FString Result = UTF8_TO_TCHAR(RecvBuffer);
+                TArray<FString> TagArray;
+                Result.ParseIntoArray(TagArray, TEXT("\n"), false);
+                int i;
+                for (i = 0; i < TagArray.Num() - 1; i++)
+                {
+                    if (TagArray[i].EndsWith(TEXT("CMDOK\r")))
+                    {
+                        bIsSuccess = true;
+                        UE_LOG(LogTemp, Log, TEXT("Socket command completed."));
+                        Stop();
+                        return 0;
+                    }
+                    else if (TagArray[i].StartsWith(TEXT("[DSDIR]")))
+                    {
+                        // just ignore the folder check
+                    }
+                    else if (TagArray[i].EndsWith(TEXT("CMDFAIL\r")))
+                    {
+                        Stop();
+                        UE_LOG(LogTemp, Error, TEXT("Socket command failed."));
+                        return 1;
+                    }
+                    else
+                    {
+                        FPlatformProcess::WritePipe(WritePipe, *TagArray[i]);
+                    }
+                }
+            }
+        }
+        double CurrentTime = FPlatformTime::Seconds();
+        if (CurrentTime - LastActivity > 10.0)
+        {
+            Stop();
+            return 1;
+        }
+        FPlatformProcess::Sleep(0.01f);
+    }
+    
+    return 0;
+}
+
+void FTcpDSCommander::Stop()
+{
+    if (DSSocket)
+    {
+        DSSocket->Close();
+        ISocketSubsystem::Get()->DestroySocket(DSSocket);
+    }
+    DSSocket = NULL;
+    
+    bStopping = true;
+}
+
+bool FTcpDSCommander::StartDSProcess()
+{
+    //bool bCreatedMutex = false;
+    //String MutexName = "Global\\DeploymentServer_Mutex_SERVERINSTANCE";
+    //Mutex DSBlockMutex = new Mutex(true, MutexName, out bCreatedMutex);
+    // is there a mutex we can use to connect test DS is running, also available on mac?
+    // there is also a failsafe mechanism for this, since the DeploymentServer will not start a new server if one is already running
+#if PLATFORM_WINDOWS
+    HANDLE mutex = CreateMutexA(NULL, TRUE, "Global\\DeploymentServer_Mutex_SERVERINSTANCE");
+    if (mutex == NULL)
+    {
+        // deployment server instance already runnning
+        return false;
+    }
+    CloseHandle(mutex);
+#endif // PLATFORM_WINDOWS
+    
+    FString DSFilename = FPaths::ConvertRelativePathToFull(FPaths::EngineDir() / TEXT("Binaries/DotNET/IOS/DeploymentServerLauncher.exe"));
+    FString WorkingFoder = FPaths::ConvertRelativePathToFull(FPaths::EngineDir() / TEXT("Binaries/DotNET/IOS/"));
+    FString Params = "";
+    
+#if PLATFORM_MAC
+    // On Mac we launch UBT with Mono
+    FString ScriptPath = FPaths::ConvertRelativePathToFull(FPaths::EngineDir() / TEXT("Build/BatchFiles/Mac/RunMono.sh"));
+    Params = FString::Printf(TEXT("\"%s\" \"%s\" %s"), *ScriptPath, *DSFilename, *Params);
+    DSFilename = TEXT("/bin/sh");
+#endif
+    UE_LOG(LogTemp, Warning, TEXT("DeploymentServer not running, Stating it!"));
+    FPlatformProcess::CreateProc(*DSFilename, *Params, true, true, true, NULL, 0, *WorkingFoder, (void*)nullptr);
+    FPlatformProcess::Sleep(1.0f);
+    
+    return true;
+}
+
+
+//**********************************************************************************************************************************************************
+//*
 FIOSTargetDevice::FIOSTargetDevice(const ITargetPlatform& InTargetPlatform)
 	: TargetPlatform(InTargetPlatform)
 	, DeviceEndpoint()
@@ -192,18 +405,9 @@ bool FIOSTargetDevice::GetUserCredentials(FString& OutUserName, FString& OutUser
 
 inline void FIOSTargetDevice::ExecuteConsoleCommand(const FString& ExecCommand) const
 {
-	FString DSFilename = FPaths::ConvertRelativePathToFull(FPaths::EngineDir() / TEXT("Binaries/DotNET/IOS/DeploymentServer.exe"));
-	FString Params = FString::Printf(TEXT(" command -device %s -param \"%s\""), *DeviceId.GetDeviceName(), *ExecCommand);
-	FString WorkingFoder = FPaths::ConvertRelativePathToFull(FPaths::EngineDir() / TEXT("Binaries/DotNET/IOS/"));
-
-#if PLATFORM_MAC
-	// On Mac we launch UBT with Mono
-	FString ScriptPath = FPaths::ConvertRelativePathToFull(FPaths::EngineDir() / TEXT("Build/BatchFiles/Mac/RunMono.sh"));
-	Params = FString::Printf(TEXT("\"%s\" \"%s\" %s"), *ScriptPath, *DSFilename, *Params);
-	DSFilename = TEXT("/bin/sh");
-#endif
-
-	FPlatformProcess::CreateProc(*DSFilename, *Params, true, false, false, NULL, 0, *WorkingFoder, NULL);
+	FString StdOut;
+	FString Params = FString::Printf(TEXT("command -device %s -param \"%s\""), *DeviceId.GetDeviceName(), *ExecCommand);
+	FIOSTargetDeviceOutput::ExecuteDSCommand(TCHAR_TO_ANSI(*Params), &StdOut);
 }
 
 inline ITargetDeviceOutputPtr FIOSTargetDevice::CreateDeviceOutputRouter(FOutputDevice* Output) const
@@ -216,3 +420,40 @@ inline ITargetDeviceOutputPtr FIOSTargetDevice::CreateDeviceOutputRouter(FOutput
 
 	return nullptr;
 }
+
+bool FIOSTargetDeviceOutput::ExecuteDSCommand(const char *CommandLine, FString* OutStdOut)
+{
+	void* WritePipe;
+	void* ReadPipe;
+	FPlatformProcess::CreatePipe(ReadPipe, WritePipe);
+	FTcpDSCommander DSCommander((uint8*)CommandLine, strlen(CommandLine), WritePipe);
+	while (DSCommander.IsValid() && !DSCommander.IsStopped())
+	{
+		FString NewLine = FPlatformProcess::ReadPipe(ReadPipe);
+		if (NewLine.Len() > 0)
+		{
+			// process the string to break it up in to lines
+			*OutStdOut += NewLine;
+		}
+
+		FPlatformProcess::Sleep(0.25);
+	}
+	FString NewLine = FPlatformProcess::ReadPipe(ReadPipe);
+	if (NewLine.Len() > 0)
+	{
+		// process the string to break it up in to lines
+		*OutStdOut += NewLine;
+	}
+
+	FPlatformProcess::Sleep(0.25);
+	FPlatformProcess::ClosePipe(ReadPipe, WritePipe);
+
+	if (!DSCommander.WasSuccess())
+	{
+		FPlatformMisc::LowLevelOutputDebugStringf(TEXT("The DeploymentServer command '%s' failed to run.\n"), ANSI_TO_TCHAR(CommandLine));
+
+		return false;
+	}
+
+	return true;
+}
diff --git a/Engine/Source/Developer/IOS/IOSTargetPlatform/Private/IOSTargetDevice.h b/Engine/Source/Developer/IOS/IOSTargetPlatform/Private/IOSTargetDevice.h
--- a/Engine/Source/Developer/IOS/IOSTargetPlatform/Private/IOSTargetDevice.h
+++ b/Engine/Source/Developer/IOS/IOSTargetPlatform/Private/IOSTargetDevice.h
@@ -5,6 +5,8 @@
 #include "CoreMinimal.h"
 #include "IMessageContext.h"
 #include "Interfaces/ITargetDevice.h"
+#include "HAL/Runnable.h"
+#include "HAL/RunnableThread.h"
 
 class FMessageEndpoint;
 class ITargetPlatform;
@@ -19,6 +21,65 @@ typedef TSharedRef<class FIOSTargetDevice, ESPMode::ThreadSafe> FIOSTargetDevice
 /** Type definition for shared references to instances of FIOSTargetDeviceOutput. */
 typedef TSharedPtr<class FIOSTargetDeviceOutput, ESPMode::ThreadSafe> FIOSTargetDeviceOutputPtr;
 
+/**
+ * Handles the communication to the Deployment Server over TCP (will start the DeploymentServer if no instance is found running)
+ */
+class FTcpDSCommander : FRunnable
+{
+public:
+    
+    /**
+     * Creates and initializes a new instance.
+     *
+     */
+    FTcpDSCommander(const uint8* Data, int32 Count, void* WPipe);
+    
+    /** Virtual destructor. */
+    virtual ~FTcpDSCommander();
+    
+public:
+    
+    //~ FRunnable interface
+    
+    virtual void Exit() override;
+    virtual bool Init() override;
+    virtual uint32 Run() override;
+    virtual void Stop() override;
+    
+    inline bool IsValid()
+    {
+        return (Thread != nullptr);
+    }
+    inline bool IsStopped()
+    {
+        return bStoped;
+    }
+    inline bool WasSuccess()
+    {
+        return bIsSuccess;
+    }
+    
+private:
+    
+    inline bool StartDSProcess();
+    
+    /** For the thread */
+    bool bStopping;
+    bool bStoped;
+    bool bIsSuccess;
+    
+    /** */
+    class FSocket* DSSocket;
+    
+    /** Holds the thread object. */
+    FRunnableThread* Thread;
+    
+    void* WritePipe;
+    uint8* DSCommand;
+    int32 DSCommandLen;
+    double LastActivity;
+};
+
 /**
  * Implements an iOS target device.
  */
diff --git a/Engine/Source/Developer/IOS/IOSTargetPlatform/Private/Windows/IOSDeviceHelperWindows.cpp b/Engine/Source/Developer/IOS/IOSTargetPlatform/Private/Windows/IOSDeviceHelperWindows.cpp
--- a/Engine/Source/Developer/IOS/IOSTargetPlatform/Private/Windows/IOSDeviceHelperWindows.cpp
+++ b/Engine/Source/Developer/IOS/IOSTargetPlatform/Private/Windows/IOSDeviceHelperWindows.cpp
@@ -1,6 +1,7 @@
 // Copyright 1998-2019 Epic Games, Inc. All Rights Reserved
 
 #include "IOSTargetPlatform.h"
+#include "IOSTargetDeviceOutput.h"
 #include "HAL/PlatformProcess.h"
 #include "HAL/Runnable.h"
 #include "HAL/RunnableThread.h"
@@ -13,98 +14,6 @@ struct FDeviceNotificationCallbackInformation
 	uint32 msgType;
 };
 
-class FDeviceSyslogRelay : public FRunnable
-{
-public:
-
-	FDeviceSyslogRelay(FString const& inDeviceID)
-	{
-		deviceID = inDeviceID;
-		Stopping = false;
-	}
-
-	virtual bool Init() override
-	{
-		return true;
-	}
-
-	virtual uint32 Run() override
-	{
-		// Tell DeploymentServer.exe to start collecting logs
-		FString ExecutablePath = FString::Printf(TEXT("%sBinaries/DotNET/IOS"), *FPaths::EngineDir());
-		FString Filename = TEXT("DeploymentServer.exe");
-
-		// Construct command line
-		FString CommandLine = FString::Printf(TEXT("listentodevice -device %s"), *deviceID);
-
-		// execute the command
-		void* WritePipe;
-		void* ReadPipe;
-		FPlatformProcess::CreatePipe(ReadPipe, WritePipe);
-		FProcHandle ProcessHandle = FPlatformProcess::CreateProc(*(ExecutablePath / Filename), *CommandLine, false, true, true, NULL, 0, *ExecutablePath, WritePipe);
-		FString	lastPartialLine;
-		while (FPlatformProcess::IsProcRunning(ProcessHandle) && !Stopping)
-		{
-			FString CurData = lastPartialLine + FPlatformProcess::ReadPipe(ReadPipe);
-			CurData.TrimStartInline();
-			if (CurData.Len() > 0)
-			{
-				// separate out each line
-				TArray<FString> LogLines;
-				CurData = CurData.Replace(TEXT("\r"), TEXT("\n"));
-
-				bool	endsWithNewline = CurData.EndsWith(TEXT("\n"));
-
-				CurData.ParseIntoArray(LogLines, TEXT("\n"), true);
-
-				for (int32 StringIndex = 0; StringIndex < LogLines.Num() - 1; ++StringIndex)
-				{
-					UE_LOG(LogIOS, Log, TEXT("%s"), *LogLines[StringIndex]);
-				}
-
-				if(endsWithNewline)
-				{
-					UE_LOG(LogIOS, Log, TEXT("%s"), *LogLines[LogLines.Num() - 1]);
-					lastPartialLine = "";
-				}
-				else
-				{
-					lastPartialLine = LogLines[LogLines.Num() - 1];
-				}
-			}
-			else
-			{
-				FPlatformProcess::Sleep(0.1);
-			}
-		}
-
-		FPlatformProcess::Sleep(0.25);
-		FString lastData = FPlatformProcess::ReadPipe(ReadPipe);
-		if (lastData.Len() > 0)
-		{
-			UE_LOG(LogIOS, Log, TEXT("%s"), *lastData);
-		}
-
-		FPlatformProcess::ClosePipe(ReadPipe, WritePipe);
-		FPlatformProcess::TerminateProc(ProcessHandle);
-
-		return 0;
-	}
-
-	virtual void Stop() override
-	{
-		Stopping = true;
-	}
-
-	virtual void Exit() override
-	{}
-
-
-private:
-	bool Stopping;
-	FString	deviceID;
-};
-
 class FIOSDevice
 {
 public:
@@ -112,29 +21,10 @@ public:
 		: UDID(InID)
 		, Name(InName)
     {
-		// BHP - Disabling the ios syslog relay because it depends on DeploymentServer running which makes it unwritable 
-		//	which causes problems when packaging because this library gets rebuilt but it can't overwrite it which
-		//	causes errors and other bad behavior - will probably need to move this functionality into its own dll
-		#if 1
-			syslogRelay = nullptr;
-			syslogRelayThread = nullptr;
-		#else
-			syslogRelay = new FDeviceSyslogRelay(InID);
-			syslogRelayThread = FRunnableThread::Create(syslogRelay, TEXT("FIOSDevice.syslogRelay"), 128 * 1024, TPri_Normal);
-		#endif
     }
     
 	~FIOSDevice()
 	{
-		if(syslogRelay != nullptr && syslogRelayThread != nullptr)
-		{
-			syslogRelay->Stop();
-			syslogRelayThread->WaitForCompletion();
-			delete syslogRelayThread;
-			delete syslogRelay;
-			syslogRelay = NULL;
-			syslogRelayThread = NULL;
-		}
 	}
 
 	FString SerialNumber() const
@@ -145,8 +35,6 @@ public:
 private:
     FString UDID;
 	FString Name;
-	FDeviceSyslogRelay*	syslogRelay;
-	FRunnableThread*	syslogRelayThread;
 };
 
 /**
@@ -210,67 +98,12 @@ public:
 	}
 
 private:
-	
-	bool ExecuteDSCommand(const FString& CommandLine, FString* OutStdOut, FString* OutStdErr) const
-	{
-		FString DSFilename = FPaths::ConvertRelativePathToFull(FPaths::EngineDir() / TEXT("Binaries/DotNET/IOS/DeploymentServer.exe"));
-		FString WorkingFoder = FPaths::ConvertRelativePathToFull(FPaths::EngineDir() / TEXT("Binaries/DotNET/IOS/"));
-
-		// execute the command
-		int32 ReturnCode;
-		FString DefaultError;
-
-		// make sure there's a place for error output to go if the caller specified nullptr
-		if (!OutStdErr)
-		{
-			OutStdErr = &DefaultError;
-		}
-
-		void* WritePipe;
-		void* ReadPipe;
-		FPlatformProcess::CreatePipe(ReadPipe, WritePipe);
-		FProcHandle ProcessHandle = FPlatformProcess::CreateProc(*DSFilename, *CommandLine, false, true, true, NULL, 0, *WorkingFoder, WritePipe);
-		while (FPlatformProcess::IsProcRunning(ProcessHandle))
-		{
-			FString NewLine = FPlatformProcess::ReadPipe(ReadPipe);
-			if (NewLine.Len() > 0)
-			{
-				// process the string to break it up in to lines
-				*OutStdOut += NewLine;
-			}
-
-			FPlatformProcess::Sleep(0.25);
-		}
-		FString NewLine = FPlatformProcess::ReadPipe(ReadPipe);
-		if (NewLine.Len() > 0)
-		{
-			// process the string to break it up in to lines
-			*OutStdOut += NewLine;
-		}
-
-		FPlatformProcess::Sleep(0.25);
-		FPlatformProcess::ClosePipe(ReadPipe, WritePipe);
-
-		if (!FPlatformProcess::GetProcReturnCode(ProcessHandle, &ReturnCode))
-		{
-			return false;
-		}
-
-		if (ReturnCode != 0)
-		{
-			FPlatformMisc::LowLevelOutputDebugStringf(TEXT("The DeploymentServer command '%s' failed to run. Return code: %d, Error: %s\n"), *CommandLine, ReturnCode, **OutStdErr);
-
-			return false;
-		}
-
-		return true;
-	}
 
 	void QueryDevices()
 	{
 		FString StdOut;
 		// get the list of devices
-		if (!ExecuteDSCommand(TEXT("listdevices"), &StdOut, nullptr))
+		if (!FIOSTargetDeviceOutput::ExecuteDSCommand("listdevices", &StdOut))
 		{
 			return;
 		}
diff --git a/Engine/Source/Developer/iOS/IOSTargetPlatform/Private/IOSTargetDeviceOutput.h b/Engine/Source/Developer/iOS/IOSTargetPlatform/Private/IOSTargetDeviceOutput.h
--- a/Engine/Source/Developer/iOS/IOSTargetPlatform/Private/IOSTargetDeviceOutput.h
+++ b/Engine/Source/Developer/iOS/IOSTargetPlatform/Private/IOSTargetDeviceOutput.h
@@ -24,7 +24,7 @@ public:
 	virtual uint32 Run(void) override;
 
 private:
-	bool StartDSProcess(void);
+	bool StartDSCommander();
 
 private:
 	// > 0 if we've been asked to abort work in progress at the next opportunity
@@ -35,7 +35,7 @@ private:
 	
 	void*				DSReadPipe;
 	void*				DSWritePipe;
-	FProcHandle			DSProcHandle;
+	FTcpDSCommander*		DSCommander;
 };
 
 /**
@@ -45,6 +45,8 @@ class FIOSTargetDeviceOutput : public ITargetDeviceOutput
 {
 public:
 	bool Init(const FIOSTargetDevice& TargetDevice, FOutputDevice* Output);
+
+	static bool ExecuteDSCommand(const char *CommandLine, FString* OutStdOut);
 	
 private:
 	TUniquePtr<FRunnableThread>						DeviceOutputThread;
diff --git a/Engine/Source/Developer/iOS/IOSTargetPlatform/Private/IOSTargetDeviceOutput.inl b/Engine/Source/Developer/iOS/IOSTargetPlatform/Private/IOSTargetDeviceOutput.inl
--- a/Engine/Source/Developer/iOS/IOSTargetPlatform/Private/IOSTargetDeviceOutput.inl
+++ b/Engine/Source/Developer/iOS/IOSTargetPlatform/Private/IOSTargetDeviceOutput.inl
@@ -3,57 +3,61 @@
 #include "CoreTypes.h"
 #include "CoreFwd.h"
 #include "Logging/LogMacros.h"
+#include "Common/TcpSocketBuilder.h"
+#include "string.h"
 
 class FIOSDeviceOutputReaderRunnable;
 class FIOSTargetDevice;
 class FIOSTargetDeviceOutput;
 
+
+
+
 inline FIOSDeviceOutputReaderRunnable::FIOSDeviceOutputReaderRunnable(const FTargetDeviceId InDeviceId, FOutputDevice* InOutput)
 	: StopTaskCounter(0)
 	, DeviceId(InDeviceId)
 	, Output(InOutput)
 	, DSReadPipe(nullptr)
 	, DSWritePipe(nullptr)
+	, DSCommander(nullptr)
 {
 }
 
-inline bool FIOSDeviceOutputReaderRunnable::StartDSProcess(void)
+inline bool FIOSDeviceOutputReaderRunnable::StartDSCommander()
 {
-	FString DSFilename = FPaths::ConvertRelativePathToFull(FPaths::EngineDir() / TEXT("Binaries/DotNET/IOS/DeploymentServer.exe"));
-	FString Params = FString::Printf(TEXT(" listentodevice -device %s"), *DeviceId.GetDeviceName());
-	FString WorkingFoder = FPaths::ConvertRelativePathToFull(FPaths::EngineDir() / TEXT("Binaries/DotNET/IOS/"));
-
-#if PLATFORM_MAC
-	// On Mac we launch UBT with Mono
-	FString ScriptPath = FPaths::ConvertRelativePathToFull(FPaths::EngineDir() / TEXT("Build/BatchFiles/Mac/RunMono.sh"));
-	Params = FString::Printf(TEXT("\"%s\" \"%s\" %s"), *ScriptPath, *DSFilename, *Params);
-	DSFilename = TEXT("/bin/sh");
-#endif
-
-	Output->Serialize(TEXT("Starting process ....."), ELogVerbosity::Log, NAME_None);
+	if (DSCommander)
+	{
+		DSCommander->Stop();
+		delete DSCommander;
+	}
+	Output->Serialize(TEXT("Starting listening ....."), ELogVerbosity::Log, NAME_None);
 	Output->Serialize(*DeviceId.GetDeviceName(), ELogVerbosity::Log, NAME_None);
-	DSProcHandle = FPlatformProcess::CreateProc(*DSFilename, *Params, true, false, false, NULL, 0, *WorkingFoder, DSWritePipe);
-	return DSProcHandle.IsValid();
+	FString Command = FString::Printf(TEXT("listentodevice -device %s"), *DeviceId.GetDeviceName());
+	uint8* DSCommand = (uint8*)TCHAR_TO_UTF8(*Command);
+	DSCommander = new FTcpDSCommander(DSCommand, strlen((const char*)DSCommand), DSWritePipe);
+	return DSCommander->IsValid();
 }
 
 inline bool FIOSDeviceOutputReaderRunnable::Init(void) 
 { 
 	FPlatformProcess::CreatePipe(DSReadPipe, DSWritePipe);
-	return StartDSProcess();
+	return StartDSCommander();
 }
 
 inline void FIOSDeviceOutputReaderRunnable::Exit(void) 
 {
 	StopTaskCounter.Increment();
-	if (DSProcHandle.IsValid())
+	if (DSCommander)
 	{
-		if (DSReadPipe && DSWritePipe)
-		{
-			FPlatformProcess::ClosePipe(DSReadPipe, DSWritePipe);
-			DSReadPipe = nullptr;
-			DSWritePipe = nullptr;
-		}
-		FPlatformProcess::TerminateProc(DSProcHandle, true);
+		DSCommander->Stop();
+		delete DSCommander;
+		DSCommander = nullptr;
+	}
+	if (DSReadPipe && DSWritePipe)
+	{
+		FPlatformProcess::ClosePipe(DSReadPipe, DSWritePipe);
+		DSReadPipe = nullptr;
+		DSWritePipe = nullptr;
 	}
 }
 
@@ -65,23 +69,21 @@ inline void FIOSDeviceOutputReaderRunnable::Stop(void)
 inline uint32 FIOSDeviceOutputReaderRunnable::Run(void)
 {
 	FString DSOutput;
-	
-	while (StopTaskCounter.GetValue() == 0 && DSProcHandle.IsValid())
+	Output->Serialize(TEXT("Starting Output"), ELogVerbosity::Log, NAME_None);
+	while (StopTaskCounter.GetValue() == 0 && DSCommander->IsValid())
 	{
-		if (!FPlatformProcess::IsProcRunning(DSProcHandle))
+		if (DSCommander->IsStopped() || !DSCommander->IsValid())
 		{
 			// When user plugs out USB cable DS process stops
 			// Keep trying to restore DS connection until code that uses this object will not kill us
 			Output->Serialize(TEXT("Trying to restore connection to device..."), ELogVerbosity::Log, NAME_None);
-			FPlatformProcess::CloseProc(DSProcHandle);
-			if (StartDSProcess())
+			if (StartDSCommander())
 			{
 				FPlatformProcess::Sleep(5.0f);
 			}
 			else
 			{
-				Output->Serialize(TEXT("Failed to start DS proccess"), ELogVerbosity::Log, NAME_None);
-				Stop();
+				Output->Serialize(TEXT("Failed to start DS commander"), ELogVerbosity::Log, NAME_None);
 			}
 		}
 		else
@@ -135,4 +137,3 @@ inline bool FIOSTargetDeviceOutput::Init(const FIOSTargetDevice& TargetDevice, F
 	DeviceOutputThread = TUniquePtr<FRunnableThread>(FRunnableThread::Create(Runnable, TEXT("FIOSDeviceOutputReaderRunnable")));
 	return true;
 }
-	
diff --git a/Engine/Source/Programs/AutomationTool/Scripts/BuildProjectCommand.Automation.cs b/Engine/Source/Programs/AutomationTool/Scripts/BuildProjectCommand.Automation.cs
--- a/Engine/Source/Programs/AutomationTool/Scripts/BuildProjectCommand.Automation.cs
+++ b/Engine/Source/Programs/AutomationTool/Scripts/BuildProjectCommand.Automation.cs
@@ -54,25 +54,6 @@ public partial class Project : CommandUtils
 			);
 	}
 
-	static Mutex DSBlockMutex = null;
-
-	static void CreateDeploymentServerMutex()
-	{
-		bool bCreatedMutex = false;
-		String MutexName = "Global\\DeploymentServer_Mutex_RestartNotAllowed";
-		DSBlockMutex = new Mutex(true, MutexName, out bCreatedMutex); // we don'e really care if ia already created here, we just want it to be active when we're building
-	}
-
-	static void CloseDeploymentServerMutex()
-	{
-		// Release the mutex to avoid the abandoned mutex timeout.
-		if (DSBlockMutex != null)
-		{
-			DSBlockMutex.ReleaseMutex();
-			DSBlockMutex.Dispose();
-			DSBlockMutex = null;
-		}
-	}
 
 	public static void Build(BuildCommand Command, ProjectParams Params, int WorkingCL = -1, ProjectBuildTargets TargetMask = ProjectBuildTargets.All)
 	{
@@ -223,9 +204,7 @@ public partial class Project : CommandUtils
 				}
 			}
 		}
-		CreateDeploymentServerMutex();
 		UE4Build.Build(Agenda, InDeleteBuildProducts: Params.Clean, InUpdateVersionFiles: WorkingCL > 0);
-		CloseDeploymentServerMutex();
 
 		if (WorkingCL > 0) // only move UAT files if we intend to check in some build products
 		{
diff --git a/Engine/Source/Programs/IOS/DeploymentServer/DeploymentServer.csproj b/Engine/Source/Programs/IOS/DeploymentServer/DeploymentServer.csproj
--- a/Engine/Source/Programs/IOS/DeploymentServer/DeploymentServer.csproj
+++ b/Engine/Source/Programs/IOS/DeploymentServer/DeploymentServer.csproj
@@ -95,10 +95,6 @@
       <Project>{B7C89A17-1CB7-42D7-B879-43297538C193}</Project>
       <Name>DeploymentInterface</Name>
     </ProjectReference>
-    <ProjectReference Include="..\DeploymentServerLauncher\DeploymentServerLauncher.csproj">
-      <Project>{0e71e8b6-849c-47c4-93b0-7e35971c6ac7}</Project>
-      <Name>DeploymentServerLauncher</Name>
-    </ProjectReference>
     <ProjectReference Include="..\MobileDeviceInterface\MobileDeviceInterface.csproj">
       <Project>{C28B2A4B-4083-423E-B6C0-48058C46AAFE}</Project>
       <Name>MobileDeviceInterface</Name>
@@ -106,7 +102,7 @@
   </ItemGroup>
   <Import Project="$(MSBuildToolsPath)\Microsoft.CSharp.targets" />
   <PropertyGroup>
-    <PreBuildEvent>StopDeploymentServer.bat
+    <PreBuildEvent>
 </PreBuildEvent>
   </PropertyGroup>
   <!-- To modify your build process, add your task inside one of the targets below and uncomment it. 
diff --git a/Engine/Source/Programs/IOS/DeploymentServer/Program.cs b/Engine/Source/Programs/IOS/DeploymentServer/Program.cs
--- a/Engine/Source/Programs/IOS/DeploymentServer/Program.cs
+++ b/Engine/Source/Programs/IOS/DeploymentServer/Program.cs
@@ -245,7 +245,7 @@ namespace DeploymentServer
 				LastResult = true;
 				bCommandComplete = false;
 				bool bWaitForCompletion = ShouldWaitForCompletion(Command);
-				bNeedsResponse = !bWaitForCompletion;
+				bNeedsResponse = bWaitForCompletion;
 
 				runLoop = new System.Threading.Thread(delegate ()
 				{
@@ -334,10 +334,13 @@ namespace DeploymentServer
 											int Ret = targetDevice.TunnelData(Param1, TCPService);
 											targetDevice.CloseTunnel(TCPService);
 
-											Writer.WriteLine("[command] Sennt '{0}' bytes.", Ret);
+											Console.WriteLine("[UE4][command] Sent '{0}' bytes. ({1})", Ret, Param1);
+											Writer.WriteLine("[UE4][command] Sent '{0}' bytes. ({1})", Ret, Param1);
 										}
+										else
 										{
-											Writer.WriteLine("[command] Device '{0}' not detected.", Device);
+											Console.WriteLine("[UE4][command] Device '{0}' not detected. ({1})", Device, Param1);
+											Writer.WriteLine("[UE4][command] Device '{0}' not detected. ({1})", Device, Param1);
 										}
 									}
 									catch
@@ -929,7 +932,7 @@ namespace DeploymentServer
 					if (TcpClientInfo.NeedsVersionCheck(LocalCommand))
 					{
 
-						if (!Response.Equals("DIR" + GetDeploymentServerPath(), StringComparison.InvariantCultureIgnoreCase))
+						if (!Response.Equals("[DSDIR]" + GetDeploymentServerPath(), StringComparison.InvariantCultureIgnoreCase))
 						{
 							Console.WriteLine("Wrong server running, restarting the server ...");
 							clientOut.Write("stop");
@@ -1143,7 +1146,7 @@ namespace DeploymentServer
 
 					TextWriter Writer = new StreamWriter(ClStream);
 					
-					Writer.WriteLine("DIR" + TestStartPath);
+					Writer.WriteLine("[DSDIR]" + TestStartPath);
 					Writer.Flush();
 
 					Byte[] Buffer = new Byte[2048];
@@ -1155,7 +1158,6 @@ namespace DeploymentServer
 						//Console.WriteLine("Looping [{0}]", localID);
 						if (ClientInfo.HasCommand && !IsStopping)
 						{
-							//Console.WriteLine("Got command [{0}]", localID);
 							if (!IsRunningCommand && !ClientInfo.IsStillRunning)
 							{
 								ClientInfo.HasCommand = false;
@@ -1173,7 +1175,10 @@ namespace DeploymentServer
 								IsRunningCommand = false;
 								if (!ClientInfo.IsStillRunning)
 								{
-									Writer.WriteLine(ClientInfo.GetLastResult ? "\nCMDOK" : "\nCMDFAIL");
+									if (ClientInfo.NeedsResponse || !ClientInfo.GetLastResult)
+									{
+										Writer.WriteLine(ClientInfo.GetLastResult ? "\nCMDOK" : "\nCMDFAIL");
+									}
 									Writer.Flush();
 									ClStream.Flush();
 									break;
@@ -1227,13 +1232,6 @@ namespace DeploymentServer
 								Client.Client.Blocking = BlockingState;
 							}
 						}
-						else if (ClientInfo.NeedsResponse)
-						{
-							Writer.WriteLine("\nCMDOK");
-							Writer.Flush();
-							ClStream.Flush();
-							break;
-						}
 						if (IsStopping)
 						{
 							try
@@ -1291,6 +1289,13 @@ namespace DeploymentServer
 					process.Kill();
 				}
 			}
+			foreach (var process in Process.GetProcessesByName("DeploymentServerLauncher"))
+			{
+				if (process.Id != Process.GetCurrentProcess().Id)
+				{
+					process.Kill();
+				}
+			}
 		}
 	}
 }
diff --git a/Engine/Source/Programs/IOS/MobileDeviceInterface/MobileDeviceInstance.cs b/Engine/Source/Programs/IOS/MobileDeviceInterface/MobileDeviceInstance.cs
--- a/Engine/Source/Programs/IOS/MobileDeviceInterface/MobileDeviceInstance.cs
+++ b/Engine/Source/Programs/IOS/MobileDeviceInterface/MobileDeviceInstance.cs
@@ -1425,9 +1425,10 @@ namespace Manzana
 			int Ret = 0;
 			try
 			{
+				Byte[] utf8Bytes = System.Text.Encoding.UTF8.GetBytes(Buffer);
 				int BufLength = Buffer.Length;
-				IntPtr intPtr_aux = Marshal.StringToHGlobalAnsi(Buffer);
-				Ret = MobileDevice.DeviceImpl.SocketSend(TCPService, intPtr_aux, BufLength);
+				Ret = TunnelBuffer(utf8Bytes, BufLength, TCPService);
+				Thread.Sleep(50);
 			}
 			catch (Exception e)
 			{
diff --git a/Engine/Source/Runtime/Launch/Private/IOS/IOSAppDelegateConsoleHandling.cpp b/Engine/Source/Runtime/Launch/Private/IOS/IOSAppDelegateConsoleHandling.cpp
--- a/Engine/Source/Runtime/Launch/Private/IOS/IOSAppDelegateConsoleHandling.cpp
+++ b/Engine/Source/Runtime/Launch/Private/IOS/IOSAppDelegateConsoleHandling.cpp
@@ -111,7 +111,9 @@ extern bool GShowSplashScreen;
 #endif
 	}
 }
+#endif
 
+#if !UE_BUILD_SHIPPING
 /**
  * Handles processing of an input console command
  */
@@ -127,7 +129,7 @@ extern bool GShowSplashScreen;
 			FPlatformString::CFStringToTCHAR((CFStringRef)ConsoleCommand, Ch.GetData());
 			new(GEngine->DeferredCommands) FString(Ch.GetData());
 		}
-
+#if !PLATFORM_TVOS
 		NSUInteger ExistingCommand = [self.ConsoleHistoryValues indexOfObjectPassingTest:
 			^ BOOL (id obj, NSUInteger idx, BOOL *stop)
 			{ 
@@ -146,6 +148,7 @@ extern bool GShowSplashScreen;
 		// save to local storage
 		[[NSUserDefaults standardUserDefaults] setObject:self.ConsoleHistoryValues forKey:@"ConsoleHistory"];
 		[[NSUserDefaults standardUserDefaults] synchronize];
+#endif
 	}
 }
 #endif
@@ -300,7 +303,7 @@ extern bool GShowSplashScreen;
 
 void EnqueueConsoleCommand(uint8 *Command)
 {
-#if !UE_BUILD_SHIPPING && !PLATFORM_TVOS
+#if !UE_BUILD_SHIPPING
 	if (!Command)
 	{
 		return;
diff --git a/Engine/Source/Runtime/Launch/Private/IOS/LaunchIOS.cpp b/Engine/Source/Runtime/Launch/Private/IOS/LaunchIOS.cpp
--- a/Engine/Source/Runtime/Launch/Private/IOS/LaunchIOS.cpp
+++ b/Engine/Source/Runtime/Launch/Private/IOS/LaunchIOS.cpp
@@ -324,7 +324,8 @@ void FAppEntry::Init()
 
 	// start up the engine
 	GEngineLoop.Init();
-#ifndef UE_BUILD_SHIPPING
+#if !UE_BUILD_SHIPPING
+	UE_LOG(LogInit, Display, TEXT("Initializing TCPConsoleListener."));
 	if (ConsoleListener == nullptr)
 	{
 		FIPv4Endpoint ConsoleTCP(FIPv4Address::InternalLoopback, 8888); //TODO: @csulea read this from some .ini
diff --git a/Engine/Source/Runtime/Launch/Private/IOS/TcpConsoleListener.cpp b/Engine/Source/Runtime/Launch/Private/IOS/TcpConsoleListener.cpp
--- a/Engine/Source/Runtime/Launch/Private/IOS/TcpConsoleListener.cpp
+++ b/Engine/Source/Runtime/Launch/Private/IOS/TcpConsoleListener.cpp
@@ -38,7 +38,6 @@ TcpConsoleListener::~TcpConsoleListener()
 
 bool TcpConsoleListener::HandleListenerConnectionAccepted(FSocket* ClientSocket, const FIPv4Endpoint& ClientEndpoint)
 {
-	UE_LOG(LogTemp, Warning, TEXT("[UE4] Got new connection!\n"));
 	Connections.Add(ClientSocket);
 	
 	return true;
@@ -84,11 +83,6 @@ uint32 TcpConsoleListener::Run()
 				if (Connection->Recv(RecvBuffer, CommandSize, BytesRead))
 				{
 					UE_LOG(LogTemp, Warning, TEXT("[UE4] Got TCP console command size %i '%s'"), BytesRead, *FString(UTF8_TO_TCHAR(RecvBuffer)));
-					int i;
-					for (i = 0; i < BytesRead; i++)
-					{
-						UE_LOG(LogTemp, Warning, TEXT("[UE4] Buffer %i"), RecvBuffer[i]);
-					}
 					EnqueueConsoleCommand(RecvBuffer);
 				}
 				else
