UE-70204 - //UE4/Main 'Build Tools' - Couldn't find parent type for 'MovieScene2DTransformTrack' named 'UMovieScenePropertyTrack'

diff --git a/Engine/Source/Developer/IOS/IOSTargetPlatform/IOSTargetPlatform.Build.cs b/Engine/Source/Developer/IOS/IOSTargetPlatform/IOSTargetPlatform.Build.cs
--- a/Engine/Source/Developer/IOS/IOSTargetPlatform/IOSTargetPlatform.Build.cs
+++ b/Engine/Source/Developer/IOS/IOSTargetPlatform/IOSTargetPlatform.Build.cs
@@ -19,8 +19,7 @@ public class IOSTargetPlatform : ModuleRules
 				"Projects",
 				"AudioPlatformConfiguration",
 				"Sockets",
-				"Networking",
-				"Engine"
+				"Networking"
 			}
 		);
 
