UE-70234 - Error: Failed to connect to deployment server at 127.0.0.1:41000 - DeploymentServer not running

diff --git a/Engine/Source/Developer/IOS/IOSTargetPlatform/Private/IOSTargetDevice.cpp b/Engine/Source/Developer/IOS/IOSTargetPlatform/Private/IOSTargetDevice.cpp
--- a/Engine/Source/Developer/IOS/IOSTargetPlatform/Private/IOSTargetDevice.cpp
+++ b/Engine/Source/Developer/IOS/IOSTargetPlatform/Private/IOSTargetDevice.cpp
@@ -29,7 +29,9 @@ FTcpDSCommander::FTcpDSCommander(const uint8* Data, int32 Count, void* WPipe)
 , DSCommandLen(Count + 1)
 , LastActivity(0.0)
 {
-    if (Count > 0)
+	// create the socket
+	ISocketSubsystem* SSS = ISocketSubsystem::Get();
+    if (SSS && Count > 0)
     {
         DSCommand = (uint8*)FMemory::Malloc(Count + 1);
         FPlatformMemory::Memcpy(DSCommand, Data, Count);
@@ -65,8 +67,7 @@ bool FTcpDSCommander::Init()
         bIsSuccess = true;
         return true;
     }
-    // create the socket
-    ISocketSubsystem* SSS = ISocketSubsystem::Get();
+	ISocketSubsystem* SSS = ISocketSubsystem::Get();
     DSSocket = SSS->CreateSocket(NAME_Stream, TEXT("DSCommander tcp"));
     if (DSSocket == nullptr)
     {
