UE-70234 - Error: Failed to connect to deployment server at 127.0.0.1:41000 - DeploymentServer not running

diff --git a/Engine/Source/Developer/IOS/IOSTargetPlatform/Private/IOSTargetDevice.cpp b/Engine/Source/Developer/IOS/IOSTargetPlatform/Private/IOSTargetDevice.cpp
--- a/Engine/Source/Developer/IOS/IOSTargetPlatform/Private/IOSTargetDevice.cpp
+++ b/Engine/Source/Developer/IOS/IOSTargetPlatform/Private/IOSTargetDevice.cpp
@@ -87,7 +87,9 @@ bool FTcpDSCommander::Init()
             // on failure, shut it all down
             SSS->DestroySocket(DSSocket);
             DSSocket = nullptr;
-            UE_LOG(LogTemp, Error, TEXT("Failed to connect to deployment server at %s."), *Addr->ToString(true));
+			ESocketErrors LastError = SSS->GetLastErrorCode();
+			const TCHAR* SocketErr = SSS->GetSocketError(LastError);
+            UE_LOG(LogTemp, Display, TEXT("Failed to connect to deployment server at %s (%s)."), *Addr->ToString(true), SocketErr);
             return false;
         }
     }
@@ -100,7 +102,7 @@ uint32 FTcpDSCommander::Run()
 {
     if (!DSSocket)
     {
-        UE_LOG(LogTemp, Error, TEXT("Socket not created."));
+        UE_LOG(LogTemp, Log, TEXT("Socket not created."));
         return 1;
     }
     int32 NSent = 0;
@@ -108,7 +110,7 @@ uint32 FTcpDSCommander::Run()
     if (NSent != DSCommandLen || !BSent)
     {
         Stop();
-        UE_LOG(LogTemp, Error, TEXT("Socket send error."));
+        UE_LOG(LogTemp, Log, TEXT("Socket send error."));
         return 1;
     }
     bStoped = false;
@@ -122,7 +124,7 @@ uint32 FTcpDSCommander::Run()
         if (DSSocket->GetConnectionState() != ESocketConnectionState::SCS_Connected)
         {
             Stop();
-            UE_LOG(LogTemp, Error, TEXT("Socket connection error."));
+            UE_LOG(LogTemp, Log, TEXT("Socket connection error."));
             return 1;
         }
         if (DSSocket->HasPendingData(Pending))
@@ -141,7 +143,7 @@ uint32 FTcpDSCommander::Run()
                     if (TagArray[i].EndsWith(TEXT("CMDOK\r")))
                     {
                         bIsSuccess = true;
-                        UE_LOG(LogTemp, Log, TEXT("Socket command completed."));
+                        //UE_LOG(LogTemp, Log, TEXT("Socket command completed."));
                         Stop();
                         return 0;
                     }
@@ -152,7 +154,7 @@ uint32 FTcpDSCommander::Run()
                     else if (TagArray[i].EndsWith(TEXT("CMDFAIL\r")))
                     {
                         Stop();
-                        UE_LOG(LogTemp, Error, TEXT("Socket command failed."));
+                        UE_LOG(LogTemp, Log, TEXT("Socket command failed."));
                         return 1;
                     }
                     else
@@ -188,9 +190,6 @@ void FTcpDSCommander::Stop()
 
 bool FTcpDSCommander::StartDSProcess()
 {
-    //bool bCreatedMutex = false;
-    //String MutexName = "Global\\DeploymentServer_Mutex_SERVERINSTANCE";
-    //Mutex DSBlockMutex = new Mutex(true, MutexName, out bCreatedMutex);
     // is there a mutex we can use to connect test DS is running, also available on mac?
     // there is also a failsafe mechanism for this, since the DeploymentServer will not start a new server if one is already running
 #if PLATFORM_WINDOWS
@@ -198,6 +197,7 @@ bool FTcpDSCommander::StartDSProcess()
     if (mutex == NULL)
     {
         // deployment server instance already runnning
+		UE_LOG(LogTemp, Display, TEXT("Deployment server instance already runnning."));
         return false;
     }
     CloseHandle(mutex);
@@ -213,7 +213,7 @@ bool FTcpDSCommander::StartDSProcess()
     Params = FString::Printf(TEXT("\"%s\" \"%s\" %s"), *ScriptPath, *DSFilename, *Params);
     DSFilename = TEXT("/bin/sh");
 #endif
-    UE_LOG(LogTemp, Warning, TEXT("DeploymentServer not running, Stating it!"));
+    UE_LOG(LogTemp, Log, TEXT("DeploymentServer not running, Starting it!"));
     FPlatformProcess::CreateProc(*DSFilename, *Params, true, true, true, NULL, 0, *WorkingFoder, (void*)nullptr);
     FPlatformProcess::Sleep(1.0f);
     
diff --git a/Engine/Source/Developer/IOS/IOSTargetPlatform/Private/Windows/IOSDeviceHelperWindows.cpp b/Engine/Source/Developer/IOS/IOSTargetPlatform/Private/Windows/IOSDeviceHelperWindows.cpp
--- a/Engine/Source/Developer/IOS/IOSTargetPlatform/Private/Windows/IOSDeviceHelperWindows.cpp
+++ b/Engine/Source/Developer/IOS/IOSTargetPlatform/Private/Windows/IOSDeviceHelperWindows.cpp
@@ -51,6 +51,7 @@ public:
 	FDeviceQueryTask()
 		: Stopping(false)
 		, bCheckDevices(true)
+		, RetryQuery(5)
 	{}
 
 	virtual bool Init() override
@@ -105,6 +106,11 @@ private:
 		// get the list of devices
 		if (!FIOSTargetDeviceOutput::ExecuteDSCommand("listdevices", &StdOut))
 		{
+			RetryQuery--;
+			if (RetryQuery < 0)
+			{
+				Enable(false);
+			}
 			return;
 		}
 
@@ -168,6 +174,7 @@ private:
 
 	bool Stopping;
 	bool bCheckDevices;
+	int RetryQuery;
 	TArray<FString> ConnectedDeviceIds;
 	FDeviceNotification DeviceNotification;
 };
