disable XGE shader compiling for IOS

diff --git a/Engine/Source/Developer/IOS/IOSTargetPlatform/Private/IOSTargetPlatform.cpp b/Engine/Source/Developer/IOS/IOSTargetPlatform/Private/IOSTargetPlatform.cpp
--- a/Engine/Source/Developer/IOS/IOSTargetPlatform/Private/IOSTargetPlatform.cpp
+++ b/Engine/Source/Developer/IOS/IOSTargetPlatform/Private/IOSTargetPlatform.cpp
@@ -493,9 +493,10 @@ static bool SupportsSoftwareOcclusion()
 
 bool FIOSTargetPlatform::CanSupportXGEShaderCompile() const
 {
+	// for 4.22 we are disabling support for XGE Shader compile on IOS
 	bool bRemoteCompilingEnabled = false;
 	GConfig->GetBool(TEXT("/Script/IOSRuntimeSettings.IOSRuntimeSettings"), TEXT("EnableRemoteShaderCompile"), bRemoteCompilingEnabled, GEngineIni);
-	return !bRemoteCompilingEnabled;
+	return false; // !bRemoteCompilingEnabled;
 }
 
 bool FIOSTargetPlatform::SupportsFeature( ETargetPlatformFeatures Feature ) const
