In order to keep up with the HLR the desktop renderers now require Metal 2.0 or later on Apple's platforms - Metal 1.2 just isn't sufficient.

diff --git a/Engine/Source/Developer/IOS/IOSPlatformEditor/Private/IOSTargetSettingsCustomization.cpp b/Engine/Source/Developer/IOS/IOSPlatformEditor/Private/IOSTargetSettingsCustomization.cpp
--- a/Engine/Source/Developer/IOS/IOSPlatformEditor/Private/IOSTargetSettingsCustomization.cpp
+++ b/Engine/Source/Developer/IOS/IOSPlatformEditor/Private/IOSTargetSettingsCustomization.cpp
@@ -1807,27 +1807,22 @@ void FIOSTargetSettingsCustomization::SetShaderStandard(int32 Value)
 			Message = LOCTEXT("iOSMetalShaderVersion1_1","Enabling Metal Shader Standard v1.1 increases the minimum operating system requirement for Metal from iOS 8.0 or later to iOS 9.0 or later. This does not affect tvOS.");
 			SetMinVersion((int32)EIOSVersion::IOS_9);
 		}
-		else if (Value < 2 && bMRTEnabled)
+		else if (Value < 3 && bMRTEnabled)
 		{
-			FPropertyAccess::Result ResMRT = ShaderVersionPropertyHandle->SetValue((uint8)2);
+			FPropertyAccess::Result ResMRT = ShaderVersionPropertyHandle->SetValue((uint8)3);
 			check(ResMRT == FPropertyAccess::Success);
 
-			Message = LOCTEXT("MetalMRTStandardv1.2","Enabling the Desktop Forward Renderer Metal requires Shader Standard v1.2 which increases the minimum operating system requirement for Metal from iOS 8.0 or later to iOS 10.0 or later.");
-			SetMinVersion((int32)EIOSVersion::IOS_10);
-		}
-		else if (Value == 2 && (EIOSVersion)EnumValue < EIOSVersion::IOS_10)
-		{
-			Message = LOCTEXT("iOSMetalShaderVersion1_2","Enabling Metal Shader Standard v1.2 increases the minimum operating system requirement for Metal from iOS 8.0/tvOS 9.0 or later to iOS/tvOS 10.0 or later.");
-			SetMinVersion((int32)EIOSVersion::IOS_10);
+			Message = LOCTEXT("MetalMRTStandardv1.2","Enabling the Desktop Forward Renderer Metal requires Shader Standard v2.0 which increases the minimum operating system requirement for Metal from iOS 10.0 or later to iOS 11.0 or later.");
+			SetMinVersion((int32)EIOSVersion::IOS_11);
 		}
 		else if (Value == 3 && (EIOSVersion)EnumValue < EIOSVersion::IOS_11)
 		{
-			Message = LOCTEXT("iOSMetalShaderVersion2_0","Enabling Metal Shader Standard v2.0 increases the minimum operating system requirement for Metal from iOS 8.0/tvOS 9.0 or later to iOS/tvOS 11.0 or later.");
+			Message = LOCTEXT("iOSMetalShaderVersion2_0","Enabling Metal Shader Standard v2.0 increases the minimum operating system requirement for Metal from iOS 10.0/tvOS 10.0 or later to iOS/tvOS 11.0 or later.");
 			SetMinVersion((int32)EIOSVersion::IOS_11);
 		}
         else if (Value == 4 && (EIOSVersion)EnumValue < EIOSVersion::IOS_12)
         {
-            Message = LOCTEXT("iOSMetalShaderVersion2_1","Enabling Metal Shader Standard v2.1 increases the minimum operating system requirement for Metal from iOS 8.0/tvOS 9.0 or later to iOS/tvOS 12.0 or later.");
+            Message = LOCTEXT("iOSMetalShaderVersion2_1","Enabling Metal Shader Standard v2.1 increases the minimum operating system requirement for Metal from iOS 10.0/tvOS 10.0 or later to iOS/tvOS 12.0 or later.");
             SetMinVersion((int32)EIOSVersion::IOS_12);
         }
 
@@ -1856,6 +1851,30 @@ void FIOSTargetSettingsCustomization::UpdateShaderStandardWarning()
 
 void FIOSTargetSettingsCustomization::UpdateOSVersionWarning()
 {
+	if (MRTPropertyHandle.IsValid() && ShaderVersionPropertyHandle.IsValid() && MinOSPropertyHandle.IsValid())
+	{
+		bool bMRTEnabled = false;
+		MRTPropertyHandle->GetValue(bMRTEnabled);
+		
+		if (bMRTEnabled)
+		{
+			uint8 EnumValue;
+			MinOSPropertyHandle->GetValue(EnumValue);
+			if (EnumValue < (uint8)EIOSVersion::IOS_11)
+			{
+				SetMinVersion((int32)EIOSVersion::IOS_11);
+				
+				FText Message;
+				Message = LOCTEXT("MetalMRTStandardv1.2","Enabling the Desktop Forward Renderer Metal requires Shader Standard v2.0 which increases the minimum operating system requirement for Metal from iOS 10.0 or later to iOS 11.0 or later.");
+				IOSVersionWarningTextBox->SetError(Message);
+			}
+		}
+		else
+		{
+			FText Message;
+			IOSVersionWarningTextBox->SetError(Message);
+		}
+	}
 }
 
 void FIOSTargetSettingsCustomization::UpdateMetalMRTWarning()
@@ -1869,20 +1888,29 @@ void FIOSTargetSettingsCustomization::UpdateMetalMRTWarning()
 		{
 			uint8 EnumValue;
 			MinOSPropertyHandle->GetValue(EnumValue);
-			if (EnumValue < (uint8)EIOSVersion::IOS_10)
+			if (EnumValue < (uint8)EIOSVersion::IOS_11)
 			{
-				SetMinVersion((int32)EIOSVersion::IOS_10);
+				SetMinVersion((int32)EIOSVersion::IOS_11);
+				
+				FText Message;
+				Message = LOCTEXT("MetalMRTStandardv1.2","Enabling the Desktop Forward Renderer Metal requires Shader Standard v2.0 which increases the minimum operating system requirement for Metal from iOS 10.0 or later to iOS 11.0 or later.");
+				IOSVersionWarningTextBox->SetError(Message);
 			}
 			
 			ShaderVersionPropertyHandle->GetValue(EnumValue);
-			if (EnumValue < (uint8)EIOSMetalShaderStandard::IOSMetalSLStandard_1_2)
+			if (EnumValue < (uint8)EIOSMetalShaderStandard::IOSMetalSLStandard_2_0)
 			{
-				SetShaderStandard((int32)EIOSMetalShaderStandard::IOSMetalSLStandard_1_2);
+				SetShaderStandard((int32)EIOSMetalShaderStandard::IOSMetalSLStandard_2_0);
+				
+				FText Message;
+				Message = LOCTEXT("MetalMRTStandardv1.2","Enabling the Desktop Forward Renderer Metal requires Shader Standard v2.0 which increases the minimum operating system requirement for Metal from iOS 10.0 or later to iOS 11.0 or later.");
+				ShaderVersionWarningTextBox->SetError(Message);
 			}
 		}
 		else
 		{
 			UpdateOSVersionWarning();
+			UpdateShaderStandardWarning();
 		}
 	}
 }
