Fixing iOS build warnings from robomerge

#rnx

diff --git a/Engine/Source/Developer/IOS/IOSTargetPlatform/Private/IOSTargetDevice.cpp b/Engine/Source/Developer/IOS/IOSTargetPlatform/Private/IOSTargetDevice.cpp
--- a/Engine/Source/Developer/IOS/IOSTargetPlatform/Private/IOSTargetDevice.cpp
+++ b/Engine/Source/Developer/IOS/IOSTargetPlatform/Private/IOSTargetDevice.cpp
@@ -73,9 +73,10 @@ bool FTcpDSCommander::Init()
     {
         return false;
     }
-    TSharedRef<FInternetAddr> Addr = SSS->CreateInternetAddr(0, DEFAULT_DS_COMMANDER_PORT);
+    TSharedRef<FInternetAddr> Addr = SSS->CreateInternetAddr();
     bool bIsValid;
     Addr->SetIp(TEXT("127.0.0.1"), bIsValid);
+	Addr->SetPort(DEFAULT_DS_COMMANDER_PORT);
     
     
     // try to connect to the server
