UE-78426 - Rare sudden crashed in editor after launching

diff --git a/Engine/Source/Developer/IOS/IOSTargetPlatform/Private/IOSTargetDevice.cpp b/Engine/Source/Developer/IOS/IOSTargetPlatform/Private/IOSTargetDevice.cpp
--- a/Engine/Source/Developer/IOS/IOSTargetPlatform/Private/IOSTargetDevice.cpp
+++ b/Engine/Source/Developer/IOS/IOSTargetPlatform/Private/IOSTargetDevice.cpp
@@ -208,7 +208,10 @@ void FTcpDSCommander::Stop()
     {
 		DSSocket->Shutdown(ESocketShutdownMode::ReadWrite);
         DSSocket->Close();
-        ISocketSubsystem::Get()->DestroySocket(DSSocket);
+		if (ISocketSubsystem::Get())
+		{
+			ISocketSubsystem::Get()->DestroySocket(DSSocket);
+		}
     }
     DSSocket = NULL;
     
