UE-77599 - [CrashReporter] UE4Editor-IOSTargetPlatform!FDeviceQueryTask::Run() [IOSDeviceHelperWindows.cpp:85]

diff --git a/Engine/Source/Developer/IOS/IOSTargetPlatform/Private/Windows/IOSDeviceHelperWindows.cpp b/Engine/Source/Developer/IOS/IOSTargetPlatform/Private/Windows/IOSDeviceHelperWindows.cpp
--- a/Engine/Source/Developer/IOS/IOSTargetPlatform/Private/Windows/IOSDeviceHelperWindows.cpp
+++ b/Engine/Source/Developer/IOS/IOSTargetPlatform/Private/Windows/IOSDeviceHelperWindows.cpp
@@ -68,9 +68,38 @@ public:
 
 	virtual uint32 Run() override
 	{
+		while (!Stopping)
+		{
+			if (GIsRequestingExit)
+			{
+				break;
+			}
+			if (GetTargetPlatformManager())
+			{
+				FString OutTutorialPath;
+				const ITargetPlatform* Platform = GetTargetPlatformManager()->FindTargetPlatform(TEXT("IOS"));
+				if (Platform)
+				{
+					if (Platform->IsSdkInstalled(false, OutTutorialPath))
+					{
+						break;
+					}
+				}
+				Enable(false);
+				return 0;
+			}
+			else
+			{
+				FPlatformProcess::Sleep(1.0f);
+			}
+		}
 		int RecheckCounter = RECHECK_COUNTER_RESET;
 		while (!Stopping)
 		{
+			if (GIsRequestingExit)
+			{
+				break;
+			}
 			if (bCheckDevices)
 			{
 #if WITH_EDITOR
@@ -78,31 +107,11 @@ public:
 				{
 					if (NeedSDKCheck)
 					{
-						if (GetTargetPlatformManager())
+						NeedSDKCheck = false;
+						FProjectStatus ProjectStatus;
+						if (!IProjectManager::Get().QueryStatusForCurrentProject(ProjectStatus) || (!ProjectStatus.IsTargetPlatformSupported(TEXT("IOS")) && !ProjectStatus.IsTargetPlatformSupported(TEXT("TVOS"))))
 						{
-							bool CanQuery = false;
-							FString OutTutorialPath;
-							const ITargetPlatform* Platform = GetTargetPlatformManager()->FindTargetPlatform(TEXT("IOS"));
-							if (Platform)
-							{
-								if (Platform->IsSdkInstalled(false, OutTutorialPath))
-								{
-									CanQuery = true;
-								}
-							}
-							NeedSDKCheck = false;
-							if (!CanQuery)
-							{
-								Enable(false);
-							}
-							else
-							{
-								FProjectStatus ProjectStatus;
-								if (!IProjectManager::Get().QueryStatusForCurrentProject(ProjectStatus) || (!ProjectStatus.IsTargetPlatformSupported(TEXT("IOS")) && !ProjectStatus.IsTargetPlatformSupported(TEXT("TVOS"))))
-								{
-									Enable(false);
-								}
-							}
+							Enable(false);
 						}
 					}
 					else
@@ -157,7 +166,7 @@ private:
 			RetryQuery--;
 			if (RetryQuery < 0 || Response < 0)
 			{
-				UE_LOG(LogTemp, Log, TEXT("IOS device listing is disabled (to many failed attempts)!"));
+				//UE_LOG(LogTemp, Log, TEXT("IOS device listing is disabled (to many failed attempts)!"));
 				Enable(false);
 			}
 			return;
