Fix race condition with FTcpDSCommander shutdown

diff --git a/Engine/Source/Developer/IOS/IOSTargetPlatform/Private/IOSTargetDevice.cpp b/Engine/Source/Developer/IOS/IOSTargetPlatform/Private/IOSTargetDevice.cpp
--- a/Engine/Source/Developer/IOS/IOSTargetPlatform/Private/IOSTargetDevice.cpp
+++ b/Engine/Source/Developer/IOS/IOSTargetPlatform/Private/IOSTargetDevice.cpp
@@ -21,7 +21,7 @@ enum
 
 FTcpDSCommander::FTcpDSCommander(const uint8* Data, int32 Count, void* WPipe)
 : bStopping(false)
-, bStoped(true)
+, bStopped(true)
 , bIsSuccess(false)
 , bIsSystemError(false)
 , DSSocket(nullptr)
@@ -54,12 +54,6 @@ FTcpDSCommander::~FTcpDSCommander()
     }
 }
 
-void FTcpDSCommander::Exit()
-{
-    // do nothing
-    bStoped = true;
-}
-
 bool FTcpDSCommander::Init()
 {
 	if (DSCommandLen < 1)
@@ -134,11 +128,10 @@ uint32 FTcpDSCommander::Run()
     bool BSent = DSSocket->Send(DSCommand, DSCommandLen, NSent);
     if (NSent != DSCommandLen || !BSent)
     {
-        Stop();
         //UE_LOG(LogTemp, Log, TEXT("Socket send error."));
         return 1;
     }
-    bStoped = false;
+    bStopped = false;
     
     static const SIZE_T CommandSize = 1024;
     uint8 RecvBuffer[CommandSize];
@@ -148,7 +141,6 @@ uint32 FTcpDSCommander::Run()
         uint32 Pending = 0;
         if (DSSocket->GetConnectionState() != ESocketConnectionState::SCS_Connected)
         {
-            Stop();
             //UE_LOG(LogTemp, Log, TEXT("Socket connection error."));
             return 1;
         }
@@ -169,7 +161,6 @@ uint32 FTcpDSCommander::Run()
                     {
                         bIsSuccess = true;
                         //UE_LOG(LogTemp, Log, TEXT("Socket command completed."));
-                        Stop();
                         return 0;
                     }
                     else if (TagArray[i].StartsWith(TEXT("[DSDIR]")))
@@ -178,7 +169,6 @@ uint32 FTcpDSCommander::Run()
                     }
                     else if (TagArray[i].EndsWith(TEXT("CMDFAIL\r")))
                     {
-                        Stop();
                         //UE_LOG(LogTemp, Display, TEXT("Socket command failed."));
                         return 1;
                     }
@@ -193,7 +183,6 @@ uint32 FTcpDSCommander::Run()
         if (CurrentTime - LastActivity > 120.0)
         {
 			//UE_LOG(LogTemp, Display, TEXT("Socket command timeouted."));
-            Stop();
             return 0;
         }
         FPlatformProcess::Sleep(0.01f);
@@ -203,19 +192,24 @@ uint32 FTcpDSCommander::Run()
 }
 
 void FTcpDSCommander::Stop()
+{ 
+    bStopping = true;
+}
+
+void FTcpDSCommander::Exit()
 {
-    if (DSSocket)
-    {
+	if (DSSocket)
+	{
 		DSSocket->Shutdown(ESocketShutdownMode::ReadWrite);
-        DSSocket->Close();
+		DSSocket->Close();
 		if (ISocketSubsystem::Get())
 		{
 			ISocketSubsystem::Get()->DestroySocket(DSSocket);
 		}
-    }
-    DSSocket = NULL;
-    
-    bStopping = true;
+	}
+	DSSocket = nullptr;
+
+	bStopped = true;
 }
 
 bool FTcpDSCommander::IsDSRunning()
diff --git a/Engine/Source/Developer/IOS/IOSTargetPlatform/Private/IOSTargetDevice.h b/Engine/Source/Developer/IOS/IOSTargetPlatform/Private/IOSTargetDevice.h
--- a/Engine/Source/Developer/IOS/IOSTargetPlatform/Private/IOSTargetDevice.h
+++ b/Engine/Source/Developer/IOS/IOSTargetPlatform/Private/IOSTargetDevice.h
@@ -55,7 +55,7 @@ public:
     }
     inline bool IsStopped()
     {
-        return bStoped;
+        return bStopped;
     }
     inline bool WasSuccess()
     {
@@ -72,7 +72,7 @@ private:
     
     /** For the thread */
     bool bStopping;
-    bool bStoped;
+    bool bStopped;
     bool bIsSuccess;
 	bool bIsSystemError; ///< Deployment server was not able to start, or connection to it could not be made
     
