Fix CIS deprecation warnings on GIsRequestingExit

diff --git a/Engine/Source/Developer/DerivedDataCache/Private/DDCCleanup.cpp b/Engine/Source/Developer/DerivedDataCache/Private/DDCCleanup.cpp
--- a/Engine/Source/Developer/DerivedDataCache/Private/DDCCleanup.cpp
+++ b/Engine/Source/Developer/DerivedDataCache/Private/DDCCleanup.cpp
@@ -124,7 +124,7 @@ uint32 FDDCCleanup::Run()
 
 bool FDDCCleanup::ShouldStop() const
 {
-	return StopTaskCounter.GetValue() > 0 || GIsRequestingExit;
+	return StopTaskCounter.GetValue() > 0 || IsEngineExitRequested();
 }
 
 bool FDDCCleanup::CleanupFilesystemDirectory( TSharedPtr< FFilesystemInfo > FilesystemInfo )
diff --git a/Engine/Source/Developer/IOS/IOSTargetPlatform/Private/Windows/IOSDeviceHelperWindows.cpp b/Engine/Source/Developer/IOS/IOSTargetPlatform/Private/Windows/IOSDeviceHelperWindows.cpp
--- a/Engine/Source/Developer/IOS/IOSTargetPlatform/Private/Windows/IOSDeviceHelperWindows.cpp
+++ b/Engine/Source/Developer/IOS/IOSTargetPlatform/Private/Windows/IOSDeviceHelperWindows.cpp
@@ -70,7 +70,7 @@ public:
 	{
 		while (!Stopping)
 		{
-			if (GIsRequestingExit)
+			if (IsEngineExitRequested())
 			{
 				break;
 			}
@@ -96,7 +96,7 @@ public:
 		int RecheckCounter = RECHECK_COUNTER_RESET;
 		while (!Stopping)
 		{
-			if (GIsRequestingExit)
+			if (IsEngineExitRequested())
 			{
 				break;
 			}
@@ -370,4 +370,4 @@ bool FIOSDeviceHelper::InstallIPAOnDevice(const FTargetDeviceId& DeviceId, const
 void FIOSDeviceHelper::EnableDeviceCheck(bool OnOff)
 {
 	QueryTask->Enable(OnOff);
-}
\ No newline at end of file
+}
diff --git a/Engine/Source/Editor/Experimental/EditorInteractiveToolsFramework/Private/EdModeInteractiveToolsContext.cpp b/Engine/Source/Editor/Experimental/EditorInteractiveToolsFramework/Private/EdModeInteractiveToolsContext.cpp
--- a/Engine/Source/Editor/Experimental/EditorInteractiveToolsFramework/Private/EdModeInteractiveToolsContext.cpp
+++ b/Engine/Source/Editor/Experimental/EditorInteractiveToolsFramework/Private/EdModeInteractiveToolsContext.cpp
@@ -833,7 +833,7 @@ void UEdModeInteractiveToolsContext::SaveEditorStateAndSetForTool()
 
 void UEdModeInteractiveToolsContext::RestoreEditorState()
 {
-	if (bHaveSavedEditorState && !GIsRequestingExit)
+	if (bHaveSavedEditorState && !IsEngineExitRequested())
 	{
 		GCurrentLevelEditingViewportClient->EngineShowFlags.SetAntiAliasing(bSavedAntiAliasingState);
 		bHaveSavedEditorState = false;
diff --git a/Engine/Source/Runtime/Engine/Private/Analytics/EditorSessionSummary.cpp b/Engine/Source/Runtime/Engine/Private/Analytics/EditorSessionSummary.cpp
--- a/Engine/Source/Runtime/Engine/Private/Analytics/EditorSessionSummary.cpp
+++ b/Engine/Source/Runtime/Engine/Private/Analytics/EditorSessionSummary.cpp
@@ -455,7 +455,7 @@ void FEditorSessionSummaryWriter::OnTerminate()
 		CurrentSession->bIsTerminating = true;
 		FPlatformMisc::SetStoredValue(SessionSummaryDefs::StoreId, CurrentSessionSectionName, SessionSummaryDefs::IsTerminatingKey, SessionSummaryDefs::TrueValueString);
 
-		if (GIsRequestingExit)
+		if (IsEngineExitRequested())
 		{
 			Shutdown();
 		}
diff --git a/Engine/Source/Runtime/Online/HTTP/Private/HttpManager.cpp b/Engine/Source/Runtime/Online/HTTP/Private/HttpManager.cpp
--- a/Engine/Source/Runtime/Online/HTTP/Private/HttpManager.cpp
+++ b/Engine/Source/Runtime/Online/HTTP/Private/HttpManager.cpp
@@ -148,7 +148,7 @@ void FHttpManager::Flush(bool bShutdown)
 			for (TArray<TSharedRef<IHttpRequest>>::TIterator It(Requests); It; ++It)
 			{
 				TSharedRef<IHttpRequest>& Request = *It;
-				if (GIsRequestingExit)
+				if (IsEngineExitRequested())
 				{
 					ensureMsgf(Request.IsUnique(), TEXT("Dangling HTTP request! This may cause undefined behaviour or crash during module shutdown!"));
 				}
