Set Max Metal Target Version default to 2.0
Marked iOS Project as a Framework as experimental

diff --git a/Engine/Config/BaseEngine.ini b/Engine/Config/BaseEngine.ini
--- a/Engine/Config/BaseEngine.ini
+++ b/Engine/Config/BaseEngine.ini
@@ -1985,7 +1985,7 @@ bAutomaticSigning=false
 UseFastIntrinsics=False
 ForceFloats=False
 EnableMathOptimisations=True
-MaxShaderLanguageVersion=2
+MaxShaderLanguageVersion=3
 bDisableMotionData=False
 bEnableAdvertisingIdentifier=True
 
diff --git a/Engine/Source/Developer/IOS/IOSPlatformEditor/Private/IOSTargetSettingsCustomization.cpp b/Engine/Source/Developer/IOS/IOSPlatformEditor/Private/IOSTargetSettingsCustomization.cpp
--- a/Engine/Source/Developer/IOS/IOSPlatformEditor/Private/IOSTargetSettingsCustomization.cpp
+++ b/Engine/Source/Developer/IOS/IOSPlatformEditor/Private/IOSTargetSettingsCustomization.cpp
@@ -874,7 +874,7 @@ void FIOSTargetSettingsCustomization::BuildPListSection(IDetailLayoutBuilder& De
 		UpdateShaderStandardWarning();
     }
 
-	// Handle max. shader version a little specially.
+	// Handle Min IOS version a little specially.
 	{
 		MinOSPropertyHandle = DetailLayout.GetProperty(GET_MEMBER_NAME_CHECKED(UIOSRuntimeSettings, MinimumiOSVersion));
 		MinOSPropertyHandle->SetOnPropertyValueChanged(OnUpdateOSVersionWarning);
@@ -1723,7 +1723,7 @@ TSharedRef<SWidget> FIOSTargetSettingsCustomization::OnGetShaderVersionContent()
 	
 	for (int32 i = 0; i < Enum->GetMaxEnumValue(); i++)
 	{
-		if (Enum->IsValidEnumValue(i))
+		if (Enum->IsValidEnumValue(i) && !Enum->HasMetaData(TEXT("Hidden"), Enum->GetIndexByValue(i)))
 		{
 			FUIAction ItemAction(FExecuteAction::CreateSP(this, &FIOSTargetSettingsCustomization::SetShaderStandard, i));
 			MenuBuilder.AddMenuEntry(Enum->GetDisplayNameTextByValue(i), TAttribute<FText>(), FSlateIcon(), ItemAction);
diff --git a/Engine/Source/Runtime/IOS/IOSRuntimeSettings/Classes/IOSRuntimeSettings.h b/Engine/Source/Runtime/IOS/IOSRuntimeSettings/Classes/IOSRuntimeSettings.h
--- a/Engine/Source/Runtime/IOS/IOSRuntimeSettings/Classes/IOSRuntimeSettings.h
+++ b/Engine/Source/Runtime/IOS/IOSRuntimeSettings/Classes/IOSRuntimeSettings.h
@@ -57,7 +57,7 @@ UENUM()
 enum class EIOSMetalShaderStandard : uint8
 {
 	/** Metal Shaders Compatible With iOS 10.0/tvOS 10.0 or later (std=ios-metal1.2) */
-	IOSMetalSLStandard_1_2 = 2 UMETA(DisplayName="Metal v1.2 (iOS 10.0/tvOS 10.0)"),
+	IOSMetalSLStandard_1_2 = 2 UMETA(DisplayName="Metal v1.2 (iOS 10.0/tvOS 10.0)", Hidden),
 	
     /** Metal Shaders Compatible With iOS 11.0/tvOS 11.0 or later (std=ios-metal2.0) */
 	IOSMetalSLStandard_2_0 = 3 UMETA(DisplayName="Metal v2.0 (iOS 11.0/tvOS 11.0)"),
@@ -241,7 +241,7 @@ public:
 	bool bCookASTCTextures;
 	
     // Whether to build the iOS project as a framework.
-    UPROPERTY(GlobalConfig, EditAnywhere, Category = Build, meta = (DisplayName = "Build project as a framework"))
+    UPROPERTY(GlobalConfig, EditAnywhere, Category = Build, meta = (DisplayName = "Build project as a framework (Experimental)"))
     bool bBuildAsFramework;
 
 	// Remotely compile shaders offline
