Fixed Device Output Log not working for iOS devices from Windows
Fixed a number of issues with Deployment Server
#jra UE-69811

diff --git a/Engine/Source/Developer/IOS/IOSTargetPlatform/Private/IOSTargetDevice.cpp b/Engine/Source/Developer/IOS/IOSTargetPlatform/Private/IOSTargetDevice.cpp
--- a/Engine/Source/Developer/IOS/IOSTargetPlatform/Private/IOSTargetDevice.cpp
+++ b/Engine/Source/Developer/IOS/IOSTargetPlatform/Private/IOSTargetDevice.cpp
@@ -19,14 +19,14 @@ enum
     DEFAULT_DS_COMMANDER_PORT = 41000, // default port to use when issuing DeploymentServer commands
 };
 
-FTcpDSCommander::FTcpDSCommander(const uint8* Data, int32 Count, void* WPipe)
+FTcpDSCommander::FTcpDSCommander(const uint8* Data, int32 Count, TQueue<FString>& InOutputQueue)
 : bStopping(false)
 , bStopped(true)
 , bIsSuccess(false)
 , bIsSystemError(false)
 , DSSocket(nullptr)
 , Thread(nullptr)
-, WritePipe(WPipe)
+, OutputQueue(InOutputQueue)
 , DSCommand(nullptr)
 , DSCommandLen(Count + 1)
 , LastActivity(0.0)
@@ -133,9 +133,11 @@ uint32 FTcpDSCommander::Run()
     }
     bStopped = false;
     
-    static const SIZE_T CommandSize = 1024;
-    uint8 RecvBuffer[CommandSize];
-    
+    static const SIZE_T BufferSize = 1025;
+	TArray<ANSICHAR> RecvBuffer;
+	RecvBuffer.Empty(BufferSize);
+	RecvBuffer.AddZeroed(BufferSize);
+
     while (!bStopping)
     {
         uint32 Pending = 0;
@@ -144,48 +146,73 @@ uint32 FTcpDSCommander::Run()
             //UE_LOG(LogTemp, Log, TEXT("Socket connection error."));
             return 1;
         }
-        if (DSSocket->HasPendingData(Pending))
-        {
-            NSent = 0;
-            FMemory::Memset(RecvBuffer, 0, CommandSize);
-            LastActivity = FPlatformTime::Seconds();
-            if (DSSocket->Recv(RecvBuffer, CommandSize, NSent))
-            {
-                FString Result = UTF8_TO_TCHAR(RecvBuffer);
-                TArray<FString> TagArray;
-                Result.ParseIntoArray(TagArray, TEXT("\n"), false);
-                int i;
-                for (i = 0; i < TagArray.Num() - 1; i++)
-                {
-                    if (TagArray[i].EndsWith(TEXT("CMDOK\r")))
-                    {
-                        bIsSuccess = true;
-                        //UE_LOG(LogTemp, Log, TEXT("Socket command completed."));
-                        return 0;
-                    }
-                    else if (TagArray[i].StartsWith(TEXT("[DSDIR]")))
-                    {
-                        // just ignore the folder check
-                    }
-                    else if (TagArray[i].EndsWith(TEXT("CMDFAIL\r")))
-                    {
-                        //UE_LOG(LogTemp, Display, TEXT("Socket command failed."));
-                        return 1;
-                    }
-                    else
-                    {
-                        FPlatformProcess::WritePipe(WritePipe, *TagArray[i]);
-                    }
-                }
-            }
-        }
-        double CurrentTime = FPlatformTime::Seconds();
-        if (CurrentTime - LastActivity > 120.0)
+
+		int32 BytesRead = 0;
+		// Receive into the last BufferSize bytes of the buffer.
+		// We only receive at max BufferSize-1 bytes to preserve the last byte for a null terminator
+		if (DSSocket->Recv((uint8*)RecvBuffer.GetData() + RecvBuffer.Num() - BufferSize, BufferSize - 1, BytesRead) == false)
+		{
+			// Recv returns false on graceful socket disconnection
+			return 1;
+		}
+
+		if (BytesRead > 0)
         {
-			//UE_LOG(LogTemp, Display, TEXT("Socket command timeouted."));
-            return 0;
-        }
-        FPlatformProcess::Sleep(0.01f);
+			LastActivity = FPlatformTime::Seconds();
+
+			ANSICHAR* LineStart = RecvBuffer.GetData();
+			ANSICHAR* LineEnd = nullptr;
+
+			while ((LineEnd = FCStringAnsi::Strchr(LineStart, '\n')) != nullptr)
+			{				
+				*LineEnd = '\0';
+				FString DataLine = UTF8_TO_TCHAR(LineStart);
+				LineStart = LineEnd + 1;
+
+				if (DataLine.EndsWith(TEXT("CMDOK\r")))
+				{
+					bIsSuccess = true;
+					return 0;
+				}
+				else if (DataLine.StartsWith(TEXT("[DSDIR]")))
+				{
+					// just ignore the folder check
+				}
+				else if (DataLine.EndsWith(TEXT("CMDFAIL\r")))
+				{
+					//UE_LOG(LogTemp, Display, TEXT("Socket command failed."));
+					return 1;
+				}
+				else
+				{
+					OutputQueue.Enqueue(DataLine);
+				}
+			}
+
+			if (*LineStart == '\0')
+			{
+				// no trailing data, reset buffer
+				RecvBuffer.Empty(BufferSize);
+				RecvBuffer.AddZeroed(BufferSize);
+			}
+			else
+			{
+				// There is trailing data, leave at the front of the buffer and add space for the next recv call.
+				RecvBuffer.RemoveAt(0, LineStart - RecvBuffer.GetData(), false);
+				// New buffer needs to have BufferSize free after the remaining string.
+				RecvBuffer.AddZeroed(BufferSize + FCStringAnsi::Strlen(RecvBuffer.GetData()) - RecvBuffer.Num());
+			}
+		}
+		else
+		{
+			double CurrentTime = FPlatformTime::Seconds();
+			if (CurrentTime - LastActivity > 120.0)
+			{
+				//UE_LOG(LogTemp, Display, TEXT("Socket command timeouted."));
+				return 0;
+			}
+			FPlatformProcess::Sleep(0.05f);
+		}
     }
     
     return 0;
@@ -462,30 +489,29 @@ inline ITargetDeviceOutputPtr FIOSTargetDevice::CreateDeviceOutputRouter(FOutput
 
 int FIOSTargetDeviceOutput::ExecuteDSCommand(const char *CommandLine, FString* OutStdOut)
 {
-	void* WritePipe;
-	void* ReadPipe;
-	FPlatformProcess::CreatePipe(ReadPipe, WritePipe);
-	FTcpDSCommander DSCommander((uint8*)CommandLine, strlen(CommandLine), WritePipe);
+	TQueue<FString>	OutputQueue;
+	FTcpDSCommander DSCommander((uint8*)CommandLine, strlen(CommandLine), OutputQueue);
 	while (DSCommander.IsValid() && !DSCommander.IsStopped())
 	{
-		FString NewLine = FPlatformProcess::ReadPipe(ReadPipe);
-		if (NewLine.Len() > 0)
+		FString NewLine;
+		if (OutputQueue.Dequeue(NewLine))
 		{
 			// process the string to break it up in to lines
 			*OutStdOut += NewLine;
 		}
-
-		FPlatformProcess::Sleep(0.25);
+		else
+		{
+			FPlatformProcess::Sleep(0.05f);
+		}
 	}
-	FString NewLine = FPlatformProcess::ReadPipe(ReadPipe);
-	if (NewLine.Len() > 0)
+	FString NewLine;
+	if (OutputQueue.Dequeue(NewLine))
 	{
 		// process the string to break it up in to lines
 		*OutStdOut += NewLine;
 	}
 
-	FPlatformProcess::Sleep(0.25);
-	FPlatformProcess::ClosePipe(ReadPipe, WritePipe);
+	FPlatformProcess::Sleep(0.05f);
 
 	if (DSCommander.IsSystemError())
 	{
diff --git a/Engine/Source/Developer/IOS/IOSTargetPlatform/Private/IOSTargetDevice.h b/Engine/Source/Developer/IOS/IOSTargetPlatform/Private/IOSTargetDevice.h
--- a/Engine/Source/Developer/IOS/IOSTargetPlatform/Private/IOSTargetDevice.h
+++ b/Engine/Source/Developer/IOS/IOSTargetPlatform/Private/IOSTargetDevice.h
@@ -32,7 +32,7 @@ public:
      * Creates and initializes a new instance.
      *
      */
-    FTcpDSCommander(const uint8* Data, int32 Count, void* WPipe);
+    FTcpDSCommander(const uint8* Data, int32 Count, TQueue<FString>& InOutputQueue);
     
     /** Virtual destructor. */
     virtual ~FTcpDSCommander();
@@ -82,7 +82,7 @@ private:
     /** Holds the thread object. */
     FRunnableThread* Thread;
     
-    void* WritePipe;
+	TQueue<FString>& OutputQueue;
     uint8* DSCommand;
     int32 DSCommandLen;
     double LastActivity;
diff --git a/Engine/Source/Developer/iOS/IOSTargetPlatform/Private/IOSTargetDeviceOutput.h b/Engine/Source/Developer/iOS/IOSTargetPlatform/Private/IOSTargetDeviceOutput.h
--- a/Engine/Source/Developer/iOS/IOSTargetPlatform/Private/IOSTargetDeviceOutput.h
+++ b/Engine/Source/Developer/iOS/IOSTargetPlatform/Private/IOSTargetDeviceOutput.h
@@ -32,10 +32,8 @@ private:
 	
 	FTargetDeviceId		DeviceId;
 	FOutputDevice*		Output;
-	
-	void*				DSReadPipe;
-	void*				DSWritePipe;
-	FTcpDSCommander*		DSCommander;
+	TQueue<FString>		OutputQueue;
+	FTcpDSCommander*	DSCommander;
 };
 
 /**
diff --git a/Engine/Source/Developer/iOS/IOSTargetPlatform/Private/IOSTargetDeviceOutput.inl b/Engine/Source/Developer/iOS/IOSTargetPlatform/Private/IOSTargetDeviceOutput.inl
--- a/Engine/Source/Developer/iOS/IOSTargetPlatform/Private/IOSTargetDeviceOutput.inl
+++ b/Engine/Source/Developer/iOS/IOSTargetPlatform/Private/IOSTargetDeviceOutput.inl
@@ -17,8 +17,6 @@ inline FIOSDeviceOutputReaderRunnable::FIOSDeviceOutputReaderRunnable(const FTar
 	: StopTaskCounter(0)
 	, DeviceId(InDeviceId)
 	, Output(InOutput)
-	, DSReadPipe(nullptr)
-	, DSWritePipe(nullptr)
 	, DSCommander(nullptr)
 {
 }
@@ -34,13 +32,12 @@ inline bool FIOSDeviceOutputReaderRunnable::StartDSCommander()
 	Output->Serialize(*DeviceId.GetDeviceName(), ELogVerbosity::Log, NAME_None);
 	FString Command = FString::Printf(TEXT("listentodevice -device %s"), *DeviceId.GetDeviceName());
 	uint8* DSCommand = (uint8*)TCHAR_TO_UTF8(*Command);
-	DSCommander = new FTcpDSCommander(DSCommand, strlen((const char*)DSCommand), DSWritePipe);
+	DSCommander = new FTcpDSCommander(DSCommand, strlen((const char*)DSCommand), OutputQueue);
 	return DSCommander->IsValid();
 }
 
 inline bool FIOSDeviceOutputReaderRunnable::Init(void) 
 { 
-	FPlatformProcess::CreatePipe(DSReadPipe, DSWritePipe);
 	return StartDSCommander();
 }
 
@@ -53,12 +50,6 @@ inline void FIOSDeviceOutputReaderRunnable::Exit(void)
 		delete DSCommander;
 		DSCommander = nullptr;
 	}
-	if (DSReadPipe && DSWritePipe)
-	{
-		FPlatformProcess::ClosePipe(DSReadPipe, DSWritePipe);
-		DSReadPipe = nullptr;
-		DSWritePipe = nullptr;
-	}
 }
 
 inline void FIOSDeviceOutputReaderRunnable::Stop(void)
@@ -68,7 +59,6 @@ inline void FIOSDeviceOutputReaderRunnable::Stop(void)
 
 inline uint32 FIOSDeviceOutputReaderRunnable::Run(void)
 {
-	FString DSOutput;
 	Output->Serialize(TEXT("Starting Output"), ELogVerbosity::Log, NAME_None);
 	while (StopTaskCounter.GetValue() == 0 && DSCommander->IsValid())
 	{
@@ -88,34 +78,18 @@ inline uint32 FIOSDeviceOutputReaderRunnable::Run(void)
 		}
 		else
 		{
-			DSOutput.Append(FPlatformProcess::ReadPipe(DSReadPipe));
-
-			if (DSOutput.Len() > 0)
+			FString Text;
+			if (OutputQueue.Dequeue(Text))
 			{
-				TArray<FString> OutputLines;
-				DSOutput.ParseIntoArray(OutputLines, TEXT("\n"), false);
-
-				if (!DSOutput.EndsWith(TEXT("\n")))
-				{
-					// partial line at the end, do not serialize it until we receive remainder
-					DSOutput = OutputLines.Last();
-					OutputLines.RemoveAt(OutputLines.Num() - 1);
-				}
-				else
+				if (Text.Contains(TEXT("[UE4]"), ESearchCase::CaseSensitive))
 				{
-					DSOutput.Reset();
-				}
-
-				for (int32 i = 0; i < OutputLines.Num(); ++i)
-				{
-					if (OutputLines[i].Contains(TEXT("[UE4]"), ESearchCase::CaseSensitive))
-					{
-						Output->Serialize(*OutputLines[i], ELogVerbosity::Log, NAME_None);
-					}
+					Output->Serialize(*Text, ELogVerbosity::Log, NAME_None);
 				}
 			}
-			
-			FPlatformProcess::Sleep(0.1f);
+			else
+			{
+				FPlatformProcess::Sleep(0.1f);
+			}
 		}
 	}
 
diff --git a/Engine/Source/Programs/IOS/DeploymentServer/DeployTime.cs b/Engine/Source/Programs/IOS/DeploymentServer/DeployTime.cs
--- a/Engine/Source/Programs/IOS/DeploymentServer/DeployTime.cs
+++ b/Engine/Source/Programs/IOS/DeploymentServer/DeployTime.cs
@@ -240,10 +240,6 @@ namespace DeploymentServer
                         ReportIF.Error("Timed out while trying to connect to a mobile device.  Make sure one is connected.");
                     }
                 }
-				catch (ThreadAbortException)
-				{
-					// we expect this to happen so we don't log it
-				}
 				catch (Exception ex)
                 {
                     ReportIF.Error(String.Format("Error encountered ('{0}') while trying to connect to a mobile device.  Please verify that iTunes is installed", ex.Message));
@@ -366,10 +362,9 @@ namespace DeploymentServer
             });
         }
 
-		public void ListenToDevice(string inDeviceID, TextWriter Writer)
+		public bool ListenToDevice(string inDeviceID, TextWriter Writer)
 		{
 			MobileDeviceInstance	targetDevice = null;
-
             PerformActionOnAllDevices(2 * StandardEnumerationDelayMS, delegate(MobileDeviceInstance Device)
             {
 				if(inDeviceID == Device.DeviceId)
@@ -384,19 +379,26 @@ namespace DeploymentServer
 				targetDevice.StartSyslogService();
 				
 				// This never returns, the process must be killed to stop logging
-				while(true)
+				while(targetDevice.IsConnected)
 				{
 					string	curLog = targetDevice.GetSyslogData();
 					if(curLog.Trim().Length > 0)
 					{
 						Writer.Write(curLog);
 					}
-					System.Threading.Thread.Sleep(50);
+					else
+					{
+						System.Threading.Thread.Sleep(50);
+					}
 				}
+
+				Writer.WriteLine("Device " + inDeviceID + " is no longer connected");
+				return true;
 			}
 			else
 			{
-				ReportIF.Error("Could not find device " + inDeviceID);
+				Writer.WriteLine("Could not find device " + inDeviceID);
+				return false;
 			}
 		}
 
diff --git a/Engine/Source/Programs/IOS/DeploymentServer/Program.cs b/Engine/Source/Programs/IOS/DeploymentServer/Program.cs
--- a/Engine/Source/Programs/IOS/DeploymentServer/Program.cs
+++ b/Engine/Source/Programs/IOS/DeploymentServer/Program.cs
@@ -107,7 +107,6 @@ namespace DeploymentServer
 			protected bool LastResult = false;
 
 			System.Threading.Thread runLoop = null;
-			System.Threading.Thread listenLoop = null;
 			bool bCommandComplete = true;
 
 			public void StopRun()
@@ -118,10 +117,6 @@ namespace DeploymentServer
 					{
 						runLoop.Abort();
 					}
-					if (listenLoop != null)
-					{
-						listenLoop.Abort();
-					}
 				}
 				catch { }
 			}
@@ -257,8 +252,7 @@ namespace DeploymentServer
 
 				LastResult = true;
 				bCommandComplete = false;
-				bool bWaitForCompletion = ShouldWaitForCompletion(Command);
-				bNeedsResponse = bWaitForCompletion;
+				bNeedsResponse = true;
 
 				runLoop = new System.Threading.Thread(delegate ()
 				{
@@ -411,30 +405,7 @@ namespace DeploymentServer
 								}
 								else
 								{
-									listenLoop = new System.Threading.Thread(delegate ()
-									{
-										try
-										{
-											DeploymentProxy.Deployer.ListenToDevice(Device, Writer);
-										}
-										catch (IOException)
-										{
-											// we expect this to happen so we don't log it
-										}
-										catch (ThreadAbortException)
-										{
-											// we expect this to happen so we don't log it
-										}
-										catch (Exception e)
-										{
-											LocalLog("Exception: " + e.ToString());
-										}
-										finally
-										{
-											listenLoop = null;
-										}
-									});
-									listenLoop.Start();
+									LastResult = DeploymentProxy.Deployer.ListenToDevice(Device, Writer);
 								}
 								break;
 						}
@@ -466,21 +437,18 @@ namespace DeploymentServer
 				try
 				{
 					runLoop.Start();
-					if (bWaitForCompletion)
+					while (!bCommandComplete)
 					{
-						while (!bCommandComplete)
-						{
-							CoreFoundationRunLoop.RunLoopRunInMode(CoreFoundationRunLoop.kCFRunLoopDefaultMode(), 1.0, 0);
-						}
-						if (runLoop.ThreadState == System.Threading.ThreadState.Running)
-						{
-							runLoop.Abort();
-						}
+						CoreFoundationRunLoop.RunLoopRunInMode(CoreFoundationRunLoop.kCFRunLoopDefaultMode(), 1.0, 0);
+						System.Threading.Thread.Sleep(100);
+					}
+					if (runLoop.ThreadState == System.Threading.ThreadState.Running)
+					{
+						runLoop.Abort();
 					}
 				}
 				catch (Exception e)
-				{
-					
+				{			
 					if (Command != "stop")
 					{
 						LocalLog("Exception: " + e.ToString());
@@ -489,8 +457,7 @@ namespace DeploymentServer
 					LastResult = false;
 				}
 				finally
-				{
-						
+				{					
 				}
 
 				return LastResult ? 0 : 1;
@@ -631,15 +598,6 @@ namespace DeploymentServer
 				}
 				return true;
 			}
-			public static bool ShouldWaitForCompletion(String Command)
-			{
-				switch (Command)
-				{
-					case "listentodevice":
-						return false;
-				}
-				return true;
-			}
 		}
 
 		static String GetDeploymentServerPath()
@@ -885,7 +843,7 @@ namespace DeploymentServer
 				while (true)
 				{
 					CoreFoundationRunLoop.RunLoopRunInMode(CoreFoundationRunLoop.kCFRunLoopDefaultMode(), 1.0, 0);
-					System.Threading.Thread.Sleep(50);
+					System.Threading.Thread.Sleep(100);
 					Writer.Flush();
 					OutSm.Flush();
 					if (TCPForwarding.Count > 0)
@@ -1013,7 +971,7 @@ namespace DeploymentServer
 		static void ClientLoop(string[] Args, string LocalCommand)
 		{
 			// on mac we only start one local instance due to mono limitations
-			if (Environment.OSVersion.Platform == PlatformID.MacOSX || Environment.OSVersion.Platform == PlatformID.Unix || Args[0].Equals("-standalone"))
+			if (Environment.OSVersion.Platform == PlatformID.MacOSX || Environment.OSVersion.Platform == PlatformID.Unix || Args.Contains("-standalone"))
 			{
 				RunLocalInstance(Args);
 				return;
@@ -1076,7 +1034,11 @@ namespace DeploymentServer
 				while (true)
 				{
 					string Response = clientIn.ReadLine();
-					if (Response.EndsWith("CMDOK"))
+					if (Response == "" || Response == null)
+					{
+						Thread.Sleep(10);
+					}
+					else if (Response.EndsWith("CMDOK"))
 					{
 						Program.ExitCode = 0;
 						break;
@@ -1090,7 +1052,6 @@ namespace DeploymentServer
 					{
 						LocalLog(Response);
 					}
-					Thread.Sleep(10);
 				}
 			}
 			catch (Exception e)
