Fix CIS

diff --git a/Engine/Source/Developer/IOS/IOSTargetPlatform/Private/IOSTargetDevice.h b/Engine/Source/Developer/IOS/IOSTargetPlatform/Private/IOSTargetDevice.h
--- a/Engine/Source/Developer/IOS/IOSTargetPlatform/Private/IOSTargetDevice.h
+++ b/Engine/Source/Developer/IOS/IOSTargetPlatform/Private/IOSTargetDevice.h
@@ -7,6 +7,7 @@
 #include "Interfaces/ITargetDevice.h"
 #include "HAL/Runnable.h"
 #include "HAL/RunnableThread.h"
+#include "Containers/Queue.h"
 
 class FMessageEndpoint;
 class ITargetPlatform;
