Fixing case of files in iOS folder

diff --git a/Engine/Source/Developer/IOS/IOSTargetPlatform/Private/IOSTargetDeviceOutput.h b/Engine/Source/Developer/IOS/IOSTargetPlatform/Private/IOSTargetDeviceOutput.h
new file mode 100644
--- /dev/null
+++ b/Engine/Source/Developer/IOS/IOSTargetPlatform/Private/IOSTargetDeviceOutput.h
@@ -0,0 +1,55 @@
+// Copyright 1998-2019 Epic Games, Inc. All Rights Reserved.
+
+#pragma once
+
+#include "Containers/UnrealString.h"
+#include "HAL/Runnable.h"
+#include "HAL/RunnableThread.h"
+#include "HAL/PlatformProcess.h"
+#include "Interfaces/ITargetDeviceOutput.h"
+#include "Misc/ConfigCacheIni.h"
+#include "HAL/ThreadSafeCounter.h"
+
+class FIOSTargetDevice;
+
+class FIOSDeviceOutputReaderRunnable : public FRunnable
+{
+public:
+	FIOSDeviceOutputReaderRunnable(const FTargetDeviceId InDeviceId, FOutputDevice* Output);
+	
+	// FRunnable interface.
+	virtual bool Init(void) override;
+	virtual void Exit(void) override; 
+	virtual void Stop(void) override;
+	virtual uint32 Run(void) override;
+
+private:
+	bool StartDSCommander();
+
+private:
+	// > 0 if we've been asked to abort work in progress at the next opportunity
+	FThreadSafeCounter	StopTaskCounter;
+	
+	FTargetDeviceId		DeviceId;
+	FOutputDevice*		Output;
+	TQueue<FString>		OutputQueue;
+	FTcpDSCommander*	DSCommander;
+};
+
+/**
+ * Implements a IOS target device output.
+ */
+class FIOSTargetDeviceOutput : public ITargetDeviceOutput
+{
+public:
+	bool Init(const FIOSTargetDevice& TargetDevice, FOutputDevice* Output);
+
+	static int ExecuteDSCommand(const char *CommandLine, FString* OutStdOut);
+	
+private:
+	TUniquePtr<FRunnableThread>						DeviceOutputThread;
+	FTargetDeviceId									DeviceId;
+	FString											DeviceName;
+};
+
+#include "IOSTargetDeviceOutput.inl"
diff --git a/Engine/Source/Developer/IOS/IOSTargetPlatform/Private/IOSTargetDeviceOutput.inl b/Engine/Source/Developer/IOS/IOSTargetPlatform/Private/IOSTargetDeviceOutput.inl
new file mode 100644
--- /dev/null
+++ b/Engine/Source/Developer/IOS/IOSTargetPlatform/Private/IOSTargetDeviceOutput.inl
@@ -0,0 +1,113 @@
+// Copyright 1998-2019 Epic Games, Inc. All Rights Reserved.
+
+#include "CoreTypes.h"
+#include "CoreFwd.h"
+#include "Logging/LogMacros.h"
+#include "Common/TcpSocketBuilder.h"
+#include "string.h"
+
+class FIOSDeviceOutputReaderRunnable;
+class FIOSTargetDevice;
+class FIOSTargetDeviceOutput;
+
+
+
+
+inline FIOSDeviceOutputReaderRunnable::FIOSDeviceOutputReaderRunnable(const FTargetDeviceId InDeviceId, FOutputDevice* InOutput)
+	: StopTaskCounter(0)
+	, DeviceId(InDeviceId)
+	, Output(InOutput)
+	, DSCommander(nullptr)
+{
+}
+
+inline bool FIOSDeviceOutputReaderRunnable::StartDSCommander()
+{
+	if (DSCommander)
+	{
+		DSCommander->Stop();
+		delete DSCommander;
+	}
+	Output->Serialize(TEXT("Starting listening ....."), ELogVerbosity::Log, NAME_None);
+	Output->Serialize(*DeviceId.GetDeviceName(), ELogVerbosity::Log, NAME_None);
+	FString Command = FString::Printf(TEXT("listentodevice -device %s"), *DeviceId.GetDeviceName());
+	uint8* DSCommand = (uint8*)TCHAR_TO_UTF8(*Command);
+	DSCommander = new FTcpDSCommander(DSCommand, strlen((const char*)DSCommand), OutputQueue);
+	return DSCommander->IsValid();
+}
+
+inline bool FIOSDeviceOutputReaderRunnable::Init(void) 
+{ 
+	return StartDSCommander();
+}
+
+inline void FIOSDeviceOutputReaderRunnable::Exit(void) 
+{
+	StopTaskCounter.Increment();
+	if (DSCommander)
+	{
+		DSCommander->Stop();
+		delete DSCommander;
+		DSCommander = nullptr;
+	}
+}
+
+inline void FIOSDeviceOutputReaderRunnable::Stop(void)
+{
+	StopTaskCounter.Increment();
+}
+
+inline uint32 FIOSDeviceOutputReaderRunnable::Run(void)
+{
+	Output->Serialize(TEXT("Starting Output"), ELogVerbosity::Log, NAME_None);
+	while (StopTaskCounter.GetValue() == 0 && DSCommander->IsValid())
+	{
+		if (DSCommander->IsStopped() || !DSCommander->IsValid())
+		{
+			// When user plugs out USB cable DS process stops
+			// Keep trying to restore DS connection until code that uses this object will not kill us
+			Output->Serialize(TEXT("Trying to restore connection to device..."), ELogVerbosity::Log, NAME_None);
+			if (StartDSCommander())
+			{
+				FPlatformProcess::Sleep(5.0f);
+			}
+			else
+			{
+				Output->Serialize(TEXT("Failed to start DS commander"), ELogVerbosity::Log, NAME_None);
+			}
+		}
+		else
+		{
+			FString Text;
+			if (OutputQueue.Dequeue(Text))
+			{
+				if (Text.Contains(TEXT("[UE4]"), ESearchCase::CaseSensitive))
+				{
+					Output->Serialize(*Text, ELogVerbosity::Log, NAME_None);
+				}
+			}
+			else
+			{
+				FPlatformProcess::Sleep(0.1f);
+			}
+		}
+	}
+
+	return 0;
+}
+
+inline bool FIOSTargetDeviceOutput::Init(const FIOSTargetDevice& TargetDevice, FOutputDevice* Output)
+{
+	check(Output);
+	// Output will be produced by background thread
+	check(Output->CanBeUsedOnAnyThread());
+	
+	DeviceId = TargetDevice.GetId();
+	DeviceName = TargetDevice.GetName();
+
+	Output->Serialize(TEXT("Creating FIOSTargetDeviceOutput ....."), ELogVerbosity::Log, NAME_None);
+		
+	auto* Runnable = new FIOSDeviceOutputReaderRunnable(DeviceId, Output);
+	DeviceOutputThread = TUniquePtr<FRunnableThread>(FRunnableThread::Create(Runnable, TEXT("FIOSDeviceOutputReaderRunnable")));
+	return true;
+}
diff --git a/Engine/Source/Developer/iOS/IOSTargetPlatform/Private/IOSTargetDeviceOutput.h b/Engine/Source/Developer/iOS/IOSTargetPlatform/Private/IOSTargetDeviceOutput.h
deleted file mode 100644
diff --git a/Engine/Source/Developer/iOS/IOSTargetPlatform/Private/IOSTargetDeviceOutput.inl b/Engine/Source/Developer/iOS/IOSTargetPlatform/Private/IOSTargetDeviceOutput.inl
deleted file mode 100644
