PR #6456: Uncommented a code line and corrected typo. (Contributed by [REDACTED_PERSON])

diff --git a/Engine/Source/Developer/IOS/IOSTargetPlatform/Private/Windows/IOSDeviceHelperWindows.cpp b/Engine/Source/Developer/IOS/IOSTargetPlatform/Private/Windows/IOSDeviceHelperWindows.cpp
--- a/Engine/Source/Developer/IOS/IOSTargetPlatform/Private/Windows/IOSDeviceHelperWindows.cpp
+++ b/Engine/Source/Developer/IOS/IOSTargetPlatform/Private/Windows/IOSDeviceHelperWindows.cpp
@@ -166,7 +166,7 @@ private:
 			RetryQuery--;
 			if (RetryQuery < 0 || Response < 0)
 			{
-				//UE_LOG(LogTemp, Log, TEXT("IOS device listing is disabled (to many failed attempts)!"));
+				UE_LOG(LogTemp, Log, TEXT("IOS device listing is disabled (too many failed attempts)!"));
 				Enable(false);
 			}
 			return;
