Minor fixes to launch screen storyboard editor UI

diff --git a/Engine/Source/Developer/IOS/IOSPlatformEditor/Private/IOSTargetSettingsCustomization.cpp b/Engine/Source/Developer/IOS/IOSPlatformEditor/Private/IOSTargetSettingsCustomization.cpp
--- a/Engine/Source/Developer/IOS/IOSPlatformEditor/Private/IOSTargetSettingsCustomization.cpp
+++ b/Engine/Source/Developer/IOS/IOSPlatformEditor/Private/IOSTargetSettingsCustomization.cpp
@@ -86,7 +86,10 @@ FIOSTargetSettingsCustomization::FIOSTargetSettingsCustomization()
 	new (IconNames)FPlatformIconInfo(TEXT("[REDACTED_EMAIL]"), LOCTEXT("AppIcon_iPadProRetina_iOS9", "iPad Pro Retina iOS9 App Icon"), FText::GetEmpty(), 167, 167, FPlatformIconInfo::Required);
 	new (IconNames)FPlatformIconInfo(TEXT("Icon1024.png"), LOCTEXT("AppIcon_Marketing", "Marketing Icon"), FText::GetEmpty(), 1024, 1024, FPlatformIconInfo::Required);
 
-	new (LaunchImageNames)FPlatformIconInfo(TEXT("LaunchScreenIOS.png"), LOCTEXT("LaunchImageIOS", "Image to be used in as Launch Screen"), FText::GetEmpty(), -1, -1, FPlatformIconInfo::Required);
+	new (LaunchImageNames)FPlatformIconInfo(TEXT("LaunchScreenIOS.png"), LOCTEXT("LaunchImageIOS", "Launch Screen Image"), LOCTEXT("LaunchImageIOSDesc", 
+		"This image is used for the Launch Screen when custom launch screen storyboards are not in use. "
+		"The image is used in both portait and landscape modes and will be uniformly scaled to occupy the full width or height as necessary for of all devices, "
+		"so if your app supports both a square image is recommended. The png file supplied must not have an alpha channel."), -1, -1, FPlatformIconInfo::Required);
 
 	bShowAllProvisions = false;
 	bShowAllCertificates = false;
@@ -1070,7 +1073,7 @@ void FIOSTargetSettingsCustomization::BuildIconSection(IDetailLayoutBuilder& Det
 	}
 
 	// Add the launch images
-	IDetailCategoryBuilder& LaunchImageCategory = DetailLayout.EditCategory(TEXT("Launch Image"));
+	IDetailCategoryBuilder& LaunchImageCategory = DetailLayout.EditCategory(FName("LaunchScreen"));
 	const FVector2D LaunchImageMaxSize(150.0f, 150.0f);
 	for (const FPlatformIconInfo& Info : LaunchImageNames)
 	{
@@ -1136,6 +1139,8 @@ void FIOSTargetSettingsCustomization::BuildImageRow(IDetailLayoutBuilder& Detail
 				SNew(STextBlock)
 				.Text(Info.IconName)
 				.Font(DetailLayout.GetDetailFont())
+				// IconDescription is not used, repurpose for tooltip
+				.ToolTipText(Info.IconDescription)
 			]
 		]
 		.ValueContent()
@@ -1148,7 +1153,6 @@ void FIOSTargetSettingsCustomization::BuildImageRow(IDetailLayoutBuilder& Detail
 			.VAlign(VAlign_Center)
 			[
 				SNew(SExternalImageReference, AutomaticImagePath, TargetImagePath)
-				.FileDescription(Info.IconDescription)
 				.RequiredSize(Info.IconRequiredSize)
 				.MaxDisplaySize(MaxDisplaySize)
 				.DeleteTargetWhenDefaultChosen(true)
diff --git a/Engine/Source/Runtime/IOS/IOSRuntimeSettings/Classes/IOSRuntimeSettings.h b/Engine/Source/Runtime/IOS/IOSRuntimeSettings/Classes/IOSRuntimeSettings.h
--- a/Engine/Source/Runtime/IOS/IOSRuntimeSettings/Classes/IOSRuntimeSettings.h
+++ b/Engine/Source/Runtime/IOS/IOSRuntimeSettings/Classes/IOSRuntimeSettings.h
@@ -425,8 +425,9 @@ public:
 	FString AdditionalPlistData;
 
 	/**
-	 * Choose whether to use a custom LaunchScreen.Storyboard as a Launchscreen. If True, the LaunchScreen.storyboard in Build/IOS/Resources/Interface
-	 * in the Engine folder will be copied to the bundle app and the Launchscreen defined in iOS Project Settings WILL NOT be included.
+	 * Choose whether to use a custom LaunchScreen.Storyboard as a Launchscreen. To use this option, create a storyboard in Xcode and 
+	 * copy it named LaunchScreen.storyboard in Build/IOS/Resources/Interface under your Project folder. This will be compiled and 
+	 * copied to the bundle app and the Launch screen image above will not be included in the app.
 	 * When using assets in your custom LaunchScreen.storyboard, add them in Build/IOS/Resources/Interface/Assets and they will be included.
 	 */
 	UPROPERTY(GlobalConfig, EditAnywhere, Category = LaunchScreen)
