Remove FullHDRCaptureData from IOS package when not packaging for MetalMRT

diff --git a/Engine/Source/Developer/IOS/IOSTargetPlatform/Private/IOSTargetPlatform.cpp b/Engine/Source/Developer/IOS/IOSTargetPlatform/Private/IOSTargetPlatform.cpp
--- a/Engine/Source/Developer/IOS/IOSTargetPlatform/Private/IOSTargetPlatform.cpp
+++ b/Engine/Source/Developer/IOS/IOSTargetPlatform/Private/IOSTargetPlatform.cpp
@@ -541,6 +541,15 @@ bool FIOSTargetPlatform::SupportsFeature( ETargetPlatformFeatures Feature ) cons
 
 #if WITH_ENGINE
 
+void FIOSTargetPlatform::GetReflectionCaptureFormats(TArray<FName>& OutFormats) const
+{
+	if (SupportsMetalMRT())
+	{
+		OutFormats.Add(FName(TEXT("FullHDR")));
+	}
+
+	OutFormats.Add(FName(TEXT("EncodedHDR")));
+}
 
 void FIOSTargetPlatform::GetAllPossibleShaderFormats( TArray<FName>& OutFormats ) const
 {
diff --git a/Engine/Source/Developer/IOS/IOSTargetPlatform/Private/IOSTargetPlatform.h b/Engine/Source/Developer/IOS/IOSTargetPlatform/Private/IOSTargetPlatform.h
--- a/Engine/Source/Developer/IOS/IOSTargetPlatform/Private/IOSTargetPlatform.h
+++ b/Engine/Source/Developer/IOS/IOSTargetPlatform/Private/IOSTargetPlatform.h
@@ -105,11 +105,7 @@ public:
 
 
 #if WITH_ENGINE
-	virtual void GetReflectionCaptureFormats( TArray<FName>& OutFormats ) const override
-	{
-		OutFormats.Add(FName(TEXT("EncodedHDR")));
-		OutFormats.Add(FName(TEXT("FullHDR")));
-	}
+	virtual void GetReflectionCaptureFormats(TArray<FName>& OutFormats) const override;
 
 	virtual void GetAllPossibleShaderFormats( TArray<FName>& OutFormats ) const override;
 
