Fixed issue preventing launch on iOS from Windows due to the latest versions of the standalone iTunes installer not setting registry keys we relied upon.

diff --git a/Engine/Source/Developer/IOS/IOSTargetPlatform/Private/IOSTargetPlatform.cpp b/Engine/Source/Developer/IOS/IOSTargetPlatform/Private/IOSTargetPlatform.cpp
--- a/Engine/Source/Developer/IOS/IOSTargetPlatform/Private/IOSTargetPlatform.cpp
+++ b/Engine/Source/Developer/IOS/IOSTargetPlatform/Private/IOSTargetPlatform.cpp
@@ -132,14 +132,32 @@ bool FIOSTargetPlatform::IsSdkInstalled(bool bProjectHasCode, FString& OutTutori
 	// Add future version checks here
 
 	// Check for iTunes 12
-	if(!biOSSDKInstalled
-		&& RegOpenKeyEx(HKEY_LOCAL_MACHINE, TEXT("SOFTWARE\\Wow6432Node\\Apple Inc.\\Apple Mobile Device Support\\Shared"), 0, KEY_READ, &hKey) == ERROR_SUCCESS
-		&& RegQueryValueEx(hKey, TEXT("MobileDeviceDLL"), 0, NULL, (BYTE*)dllPath, &pathSize) == ERROR_SUCCESS
-		&&  IFileManager::Get().FileSize(*FString(dllPath)) != INDEX_NONE)
+	if (!biOSSDKInstalled && RegOpenKeyEx(HKEY_LOCAL_MACHINE, TEXT("SOFTWARE\\Wow6432Node\\Apple Inc.\\Apple Mobile Device Support\\Shared"), 0, KEY_READ, &hKey) == ERROR_SUCCESS)
+	{
+		if (RegQueryValueEx(hKey, TEXT("MobileDeviceDLL"), 0, NULL, (BYTE*)dllPath, &pathSize) == ERROR_SUCCESS
+			&& IFileManager::Get().FileSize(*FString(dllPath)) != INDEX_NONE)
 		{
-		biOSSDKInstalled = true;
+			biOSSDKInstalled = true;
+		}
+		else
+		{
+			// iTunes >= 12.7 doesn't have a key specifying the 32-bit DLL but it does have a ASMapiInterfaceDLL key and MobileDevice.dll is in usually in the same directory
+			if (RegQueryValueEx(hKey, TEXT("ASMapiInterfaceDLL"), 0, NULL, (BYTE*)dllPath, &pathSize) == ERROR_SUCCESS)
+			{
+				FString MobileDeviceDLLPath = FString(dllPath);
+				int32 Index = MobileDeviceDLLPath.Find(TEXT("\\"), ESearchCase::CaseSensitive, ESearchDir::FromEnd);
+				if(Index != INDEX_NONE)
+				{
+					MobileDeviceDLLPath = MobileDeviceDLLPath.Left(Index+1) + TEXT("MobileDevice.dll");
+					if (IFileManager::Get().FileSize(*MobileDeviceDLLPath) != INDEX_NONE)
+					{
+						biOSSDKInstalled = true;
+					}
+				}
+			}
+		}
 	}
-	
+
 	// Check for iTunes 12, Windows Store version
 	if (!biOSSDKInstalled)
 	{
diff --git a/Engine/Source/Programs/IOS/MobileDeviceInterface/MobileDevice.cs b/Engine/Source/Programs/IOS/MobileDeviceInterface/MobileDevice.cs
--- a/Engine/Source/Programs/IOS/MobileDeviceInterface/MobileDevice.cs
+++ b/Engine/Source/Programs/IOS/MobileDeviceInterface/MobileDevice.cs
@@ -3023,7 +3023,14 @@ namespace Manzana
 				string dllPath12 = Microsoft.Win32.Registry.GetValue("HKEY_LOCAL_MACHINE\\SOFTWARE\\Wow6432Node\\Apple Inc.\\Apple Mobile Device Support\\Shared", "MobileDeviceDLL", null) as string;
 				string dllPath12WindowsStore = ITunesFinder.FindWindowsStoreITunesDLL();
 
-				if ((!String.IsNullOrEmpty(dllPath12) && File.Exists(dllPath12)) || (!String.IsNullOrEmpty(dllPath12WindowsStore) && File.Exists(dllPath12WindowsStore)))
+				// iTunes >= 12.7 doesn't have a key specifying the 32-bit DLL but it does have a ASMapiInterfaceDLL key and MobileDevice.dll is in usually in the same directory
+				string dllPath12_7 = Microsoft.Win32.Registry.GetValue("HKEY_LOCAL_MACHINE\\SOFTWARE\\Wow6432Node\\Apple Inc.\\Apple Mobile Device Support\\Shared", "ASMapiInterfaceDLL", null) as string;
+				if (!String.IsNullOrEmpty(dllPath12_7))
+				{
+					dllPath12_7 = dllPath12_7.Substring(0, dllPath12_7.LastIndexOf('\\') + 1) + "MobileDevice.dll";
+				}
+
+				if ((!String.IsNullOrEmpty(dllPath12) && File.Exists(dllPath12)) || (!String.IsNullOrEmpty(dllPath12_7) && File.Exists(dllPath12_7)) || (!String.IsNullOrEmpty(dllPath12WindowsStore) && File.Exists(dllPath12WindowsStore)))
 				{
 					DeviceImpl = new MobileDeviceWiniTunes12();
 				}
diff --git a/Engine/Source/Programs/IOS/iPhonePackager/iPhonePackager.cs b/Engine/Source/Programs/IOS/iPhonePackager/iPhonePackager.cs
--- a/Engine/Source/Programs/IOS/iPhonePackager/iPhonePackager.cs
+++ b/Engine/Source/Programs/IOS/iPhonePackager/iPhonePackager.cs
@@ -869,7 +869,14 @@ namespace iPhonePackager
                                 dllPath = Registry.GetValue("HKEY_LOCAL_MACHINE\\SOFTWARE\\Wow6432Node\\Apple Inc.\\Apple Mobile Device Support\\Shared", "MobileDeviceDLL", null) as string;
                                 if (String.IsNullOrEmpty(dllPath) || !File.Exists(dllPath))
                                 {
-                                    dllPath = FindWindowsStoreITunesDLL();
+									// iTunes >= 12.7 doesn't have a key specifying the 32-bit DLL but it does have a ASMapiInterfaceDLL key and MobileDevice.dll is in usually in the same directory
+									dllPath = Registry.GetValue("HKEY_LOCAL_MACHINE\\SOFTWARE\\Wow6432Node\\Apple Inc.\\Apple Mobile Device Support\\Shared", "ASMapiInterfaceDLL", null) as string;
+									dllPath = String.IsNullOrEmpty(dllPath) ? null : dllPath.Substring(0, dllPath.LastIndexOf('\\') + 1) + "MobileDevice.dll";
+
+									if (String.IsNullOrEmpty(dllPath) || !File.Exists(dllPath))
+									{
+										dllPath = FindWindowsStoreITunesDLL();
+									}
                                 }
                             }
                         }
diff --git a/Engine/Source/Programs/UnrealBuildTool/Platform/Mac/ApplePlatformSDK.cs b/Engine/Source/Programs/UnrealBuildTool/Platform/Mac/ApplePlatformSDK.cs
--- a/Engine/Source/Programs/UnrealBuildTool/Platform/Mac/ApplePlatformSDK.cs
+++ b/Engine/Source/Programs/UnrealBuildTool/Platform/Mac/ApplePlatformSDK.cs
@@ -58,8 +58,14 @@ namespace UnrealBuildTool
 					DllPath = Registry.GetValue("HKEY_LOCAL_MACHINE\\SOFTWARE\\Wow6432Node\\Apple Inc.\\Apple Mobile Device Support\\Shared", "MobileDeviceDLL", null) as string;
 					if (string.IsNullOrEmpty(DllPath) || !File.Exists(DllPath))
 					{
-						DllPath = FindWindowsStoreITunesDLL();
+						// iTunes >= 12.7 doesn't have a key specifying the 32-bit DLL but it does have a ASMapiInterfaceDLL key and MobileDevice.dll is in usually in the same directory
+						DllPath = Registry.GetValue("HKEY_LOCAL_MACHINE\\SOFTWARE\\Wow6432Node\\Apple Inc.\\Apple Mobile Device Support\\Shared", "ASMapiInterfaceDLL", null) as string;
+						DllPath = String.IsNullOrEmpty(DllPath) ? null : DllPath.Substring(0, DllPath.LastIndexOf('\\') + 1) + "MobileDevice.dll";
 
+						if (string.IsNullOrEmpty(DllPath) || !File.Exists(DllPath))
+						{
+							DllPath = FindWindowsStoreITunesDLL();
+						}
 					}
 				}
 
