- C# Turnkey now uses ID, not Name, for proper connection to editor turnkey calls
- Changed RunLocal* functions in UnrealBuildTools.Utils to use UTF-8 stdout
- IOS device detection (TVOS needs some testing still)

diff --git a/Engine/Source/Developer/IOS/IOSTargetPlatform/Private/Windows/IOSDeviceHelperWindows.cpp b/Engine/Source/Developer/IOS/IOSTargetPlatform/Private/Windows/IOSDeviceHelperWindows.cpp
--- a/Engine/Source/Developer/IOS/IOSTargetPlatform/Private/Windows/IOSDeviceHelperWindows.cpp
+++ b/Engine/Source/Developer/IOS/IOSTargetPlatform/Private/Windows/IOSDeviceHelperWindows.cpp
@@ -19,6 +19,67 @@ struct FDeviceNotificationCallbackInformation
 	uint32 msgType;
 };
 
+
+	struct LibIMobileDevice
+	{
+		FString DeviceName;
+		FString DeviceID;
+		FString DeviceType;
+	};
+
+	static TArray<LibIMobileDevice> GetLibIMobileDevices()
+	{
+		FString OutStdOut;
+		FString OutStdErr;
+		FString LibimobileDeviceId = FPaths::ConvertRelativePathToFull(FPaths::EngineDir() / TEXT("Extras/ThirdPartyNotUE/libimobiledevice/x64/idevice_id.exe"));
+		int ReturnCode;
+		// get the list of devices UDID
+		FPlatformProcess::ExecProcess(*LibimobileDeviceId, TEXT(""), &ReturnCode, &OutStdOut, &OutStdErr);
+
+		TArray<LibIMobileDevice> ToReturn;
+
+		// separate out each line
+		TArray<FString> CurrentDeviceIds;
+		OutStdOut.ParseIntoArray(CurrentDeviceIds, TEXT("\n"), true);
+		for (int32 Index = 0; Index != CurrentDeviceIds.Num(); ++Index)
+		{
+			if (CurrentDeviceIds[Index].Contains("Network"))
+			{
+				CurrentDeviceIds.RemoveAt(Index);
+				continue;
+			}
+			CurrentDeviceIds[Index].Split(TEXT(" "), &CurrentDeviceIds[Index], nullptr, ESearchCase::CaseSensitive, ESearchDir::FromStart);
+		}
+		TArray<FString> DeviceStrings;
+		for (int32 StringIndex = 0; StringIndex < CurrentDeviceIds.Num(); ++StringIndex)
+		{
+
+			const FString& DeviceID = CurrentDeviceIds[StringIndex];
+
+			FString OutStdOutInfo;
+			FString OutStdErrInfo;
+			FString LibimobileDeviceInfo = FPaths::ConvertRelativePathToFull(FPaths::EngineDir() / TEXT("Extras/ThirdPartyNotUE/libimobiledevice/x64/ideviceinfo.exe"));
+			int ReturnCodeInfo;
+			FString Arguments = "-u " + DeviceID;
+			// get the list of devices UDID
+			FPlatformProcess::ExecProcess(*LibimobileDeviceInfo, *Arguments, &ReturnCodeInfo, &OutStdOutInfo, &OutStdErrInfo);
+
+			// parse product type and device name
+			FString DeviceName;
+			OutStdOutInfo.Split(TEXT("DeviceName: "), nullptr, &DeviceName, ESearchCase::CaseSensitive, ESearchDir::FromStart);
+			DeviceName.Split(TEXT("\r\n"), &DeviceName, nullptr, ESearchCase::CaseSensitive, ESearchDir::FromStart);
+			FString ProductType;
+			OutStdOutInfo.Split(TEXT("ProductType: "), nullptr, &ProductType, ESearchCase::CaseSensitive, ESearchDir::FromStart);
+			ProductType.Split(TEXT("\r\n"), &ProductType, nullptr, ESearchCase::CaseSensitive, ESearchDir::FromStart);
+			LibIMobileDevice ToAdd;
+			ToAdd.DeviceID = DeviceID;
+			ToAdd.DeviceName = DeviceName;
+			ToAdd.DeviceType = ProductType;
+			ToReturn.Add(ToAdd);
+		}
+		return ToReturn;
+	}
+
 class FIOSDevice
 {
 public:
@@ -107,16 +168,16 @@ public:
 #if WITH_EDITOR
 				if (!IsRunningCommandlet())
 				{
-					if (NeedSDKCheck)
-					{
-						NeedSDKCheck = false;
-						FProjectStatus ProjectStatus;
-						if (!IProjectManager::Get().QueryStatusForCurrentProject(ProjectStatus) || (!ProjectStatus.IsTargetPlatformSupported(TEXT("IOS")) && !ProjectStatus.IsTargetPlatformSupported(TEXT("TVOS"))))
-						{
-							Enable(false);
-						}
-					}
-					else
+					//if (NeedSDKCheck)
+					//{
+					//	NeedSDKCheck = false;
+					//	FProjectStatus ProjectStatus;
+					//	if (!IProjectManager::Get().QueryStatusForCurrentProject(ProjectStatus) || (!ProjectStatus.IsTargetPlatformSupported(TEXT("IOS")) && !ProjectStatus.IsTargetPlatformSupported(TEXT("TVOS"))))
+					//	{
+					//		Enable(false);
+					//	}
+					//}
+					//else
 					{
 						// BHP - Turning off device check to prevent it from interfering with packaging
 						QueryDevices();
@@ -160,13 +221,16 @@ private:
 
 	void QueryDevices()
 	{
-		FString StdOut;
-		// get the list of devices
-		int Response = FIOSTargetDeviceOutput::ExecuteDSCommand("listdevices", &StdOut);
-		if (Response <= 0)
+		FString OutStdOut;
+		FString OutStdErr;
+		FString LibimobileDeviceId = FPaths::ConvertRelativePathToFull(FPaths::EngineDir() / TEXT("Extras/ThirdPartyNotUE/libimobiledevice/x64/idevice_id.exe"));
+		int ReturnCode;
+		// get the list of devices UDID
+		FPlatformProcess::ExecProcess(*LibimobileDeviceId, TEXT(""), &ReturnCode, &OutStdOut, &OutStdErr);
+		if (OutStdOut.Len() == 0)
 		{
 			RetryQuery--;
-			if (RetryQuery < 0 || Response < 0)
+			if (RetryQuery < 0)
 			{
 				UE_LOG(LogIOSDeviceHelper, Verbose, TEXT("IOS device listing is disabled for 1 minute (too many failed attempts)!"));
 				Enable(false);
@@ -174,53 +238,27 @@ private:
 			return;
 		}
 		RetryQuery = 5;
-		// separate out each line
-		TArray<FString> DeviceStrings;
-		StdOut = StdOut.Replace(TEXT("\r"), TEXT("\n"));
-		StdOut.ParseIntoArray(DeviceStrings, TEXT("\n"), true);
+	
+		TArray<LibIMobileDevice> ParsedDevices = GetLibIMobileDevices();
 
-		TArray<FString> CurrentDeviceIds;
-		for (int32 StringIndex = 0; StringIndex < DeviceStrings.Num(); ++StringIndex)
+		for (int32 Index = 0; Index < ParsedDevices.Num(); ++Index)
 		{
-			const FString& DeviceString = DeviceStrings[StringIndex];
-
-			if(!DeviceString.StartsWith("[DD] FOUND: "))
-			{
-				continue;
-			}
-
-			// grab the device serial number
-			int32 typeIndex = DeviceString.Find(TEXT("TYPE: "));
-			int32 idIndex = DeviceString.Find(TEXT("ID: "));
-			int32 nameIndex = DeviceString.Find(TEXT("NAME: "));
-			if (typeIndex < 0 || idIndex < 0 || nameIndex < 0)
-			{
-				continue;
-			}
-
-			FString SerialNumber = DeviceString.Mid(idIndex + 4, nameIndex - 1 - (idIndex + 4));
-			CurrentDeviceIds.Add(SerialNumber);
-
 			// move on to next device if this one is already a known device
-			if (ConnectedDeviceIds.Find(SerialNumber) != INDEX_NONE)
+			if (ConnectedDeviceIds.Find(ParsedDevices[Index].DeviceID) != INDEX_NONE)
 			{
-				ConnectedDeviceIds.Remove(SerialNumber);
+				ConnectedDeviceIds.Remove(ParsedDevices[Index].DeviceID);
 				continue;
 			}
 
-			// parse product type and device name
-			FString ProductType = DeviceString.Mid(typeIndex + 6, idIndex - 1 - (typeIndex + 6));
-			FString DeviceName = DeviceString.Mid(nameIndex + 6, DeviceString.Len() - (nameIndex + 6));
-
 			// create an FIOSDevice
 			FDeviceNotificationCallbackInformation CallbackInfo;
-			CallbackInfo.DeviceName = DeviceName;
-			CallbackInfo.UDID = SerialNumber;
-			CallbackInfo.ProductType = ProductType;
+			CallbackInfo.DeviceName = ParsedDevices[Index].DeviceName;
+			CallbackInfo.UDID = ParsedDevices[Index].DeviceID;
+			CallbackInfo.ProductType = ParsedDevices[Index].DeviceType;
 			CallbackInfo.msgType = 1;
 			DeviceNotification.Broadcast(&CallbackInfo);
 		}
-
+		
 		// remove all devices no longer found
 		for (int32 DeviceIndex = 0; DeviceIndex < ConnectedDeviceIds.Num(); ++DeviceIndex)
 		{
@@ -229,7 +267,11 @@ private:
 			CallbackInfo.msgType = 2;
 			DeviceNotification.Broadcast(&CallbackInfo);
 		}
-		ConnectedDeviceIds = CurrentDeviceIds;
+		ConnectedDeviceIds.Empty();
+		for (int32 Index = 0; Index < ParsedDevices.Num(); ++Index)
+		{
+			ConnectedDeviceIds.Add(ParsedDevices[Index].DeviceID);
+		}
 	}
 
 	bool Stopping;
@@ -264,18 +306,6 @@ bool FIOSDeviceHelper::MessageTickDelegate(float DeltaTime)
 
 void FIOSDeviceHelper::Initialize(bool bIsTVOS)
 {
-	// Create a dummy device to hand over
-	const FString DummyDeviceName = FString::Printf(TEXT("All_%s_On_%s"), bIsTVOS ? TEXT("tvOS") : TEXT("iOS"), FPlatformProcess::ComputerName());
-	
-	FIOSLaunchDaemonPong Event;
-	Event.DeviceID = FString::Printf(TEXT("%s@%s"), bIsTVOS ? TEXT("TVOS") : TEXT("IOS"), *DummyDeviceName);
-	Event.bCanReboot = false;
-	Event.bCanPowerOn = false;
-	Event.bCanPowerOff = false;
-	Event.DeviceName = DummyDeviceName;
-	Event.DeviceType = bIsTVOS ? TEXT("AppleTV") : TEXT("");
-	FIOSDeviceHelper::OnDeviceConnected().Broadcast(Event);
-
 	if(!bIsTVOS)
 	{
 		// add the message pump
diff --git a/Engine/Source/Programs/AutomationTool/IOS/IOSPlatform.Automation.cs b/Engine/Source/Programs/AutomationTool/IOS/IOSPlatform.Automation.cs
--- a/Engine/Source/Programs/AutomationTool/IOS/IOSPlatform.Automation.cs
+++ b/Engine/Source/Programs/AutomationTool/IOS/IOSPlatform.Automation.cs
@@ -501,31 +501,69 @@ public class IOSPlatform : Platform
 
 	public override DeviceInfo[] GetDevices()
 	{
-		if (HostPlatform.Current.HostEditorPlatform != UnrealTargetPlatform.Mac)
-		{
-			return base.GetDevices();
-		}
+		List<DeviceInfo> Devices = new List<DeviceInfo>();
 
-		string DeviceType = TargetPlatformType == UnrealTargetPlatform.TVOS ? "tvOS" : "iOS";
+		if (HostPlatform.Current.HostEditorPlatform == UnrealTargetPlatform.Mac)
+		{
+			string DeviceType = TargetPlatformType == UnrealTargetPlatform.TVOS ? "tvOS" : "iOS";
 
-		// print out each connected device's udid, os version, and name (which might have spaces)
-		string Params = "-e 'require \"fastlane\"; FastlaneCore::DeviceManager.connected_devices(\"" + DeviceType + "\").each { |x| puts \"#{x.udid} #{x.os_type} #{x.os_version} #{x.name}\" }'";
-		string StdOut = UnrealBuildTool.Utils.RunLocalProcessAndReturnStdOut("ruby", Params);
+			// print out each connected device's udid, os version, and name (which might have spaces)
+			string Params = "-e 'require \"fastlane\"; FastlaneCore::DeviceManager.connected_devices(\"" + DeviceType + "\").each { |x| puts \"#{x.udid} #{x.os_type} #{x.os_version} #{x.name}\" }'";
+			string StdOut = UnrealBuildTool.Utils.RunLocalProcessAndReturnStdOut("ruby", Params);
 
-		List<DeviceInfo> Devices = new List<DeviceInfo>();
-		foreach (string Line in StdOut.Split("\n".ToCharArray(), StringSplitOptions.RemoveEmptyEntries))
-		{
-			string[] Tokens = Line.Split(" ".ToCharArray());
+			foreach (string Line in StdOut.Split("\n".ToCharArray(), StringSplitOptions.RemoveEmptyEntries))
+			{
+				string[] Tokens = Line.Split(" ".ToCharArray());
 
-			DeviceInfo Device = new DeviceInfo(TargetPlatformType);
-			Device.Id = Tokens[0];
-			Device.Type = Tokens[1];
-			Device.SoftwareVersion = Tokens[2];
-			Device.Name = string.Join(" ", Tokens.Skip(3));
+				DeviceInfo Device = new DeviceInfo(TargetPlatformType);
+				Device.Id = Tokens[0];
+				Device.Type = Tokens[1];
+				Device.SoftwareVersion = Tokens[2];
+				Device.Name = string.Join(" ", Tokens.Skip(3));
 
-			Devices.Add(Device);
+				Devices.Add(Device);
+			}
 		}
+		else if (HostPlatform.Current.HostEditorPlatform == UnrealTargetPlatform.Win64)
+		{
+			string UtilPath = Path.Combine(CommandUtils.EngineDirectory.FullName, "Extras/ThirdPartyNotUE/libimobiledevice/x64/ideviceinfo");
+			string Output = Utils.RunLocalProcessAndReturnStdOut(UtilPath, "");
 
+			DeviceInfo CurrentDevice = null;
+			foreach (string Line in Output.Split("\r\n".ToCharArray()))
+			{
+				if (Line.StartsWith("ActivationState:"))
+				{
+					CurrentDevice = new DeviceInfo(TargetPlatformType);
+					Devices.Add(CurrentDevice);
+				}
+				else if (Line.StartsWith("DeviceName:"))
+				{
+					CurrentDevice.Name = Line.Split(": ").Last();
+				}
+				// check we are returning the proper device for this class
+				else if (Line.StartsWith("DeviceClass:"))
+				{
+					bool bIsDeviceTVOS = Line.Split(": ").Last().ToLower() == "tvos";
+					if (bIsDeviceTVOS != (TargetPlatformType == UnrealTargetPlatform.TVOS))
+					{
+						Devices.Remove(CurrentDevice);
+					}
+				}
+				else if (Line.StartsWith("UniqueDeviceID: "))
+				{
+					CurrentDevice.Id = Line.Split(": ").Last();
+				}
+				else if (Line.StartsWith("ProductType: "))
+				{
+					CurrentDevice.Type = Line.Split(": ").Last();
+				}
+				else if (Line.StartsWith("ProductVersion: "))
+				{
+					CurrentDevice.SoftwareVersion = Line.Split(": ").Last();
+				}
+			}
+		}
 		return Devices.ToArray();
 	}
 
diff --git a/Engine/Source/Programs/AutomationTool/Turnkey/Commands/VerifySdk.cs b/Engine/Source/Programs/AutomationTool/Turnkey/Commands/VerifySdk.cs
--- a/Engine/Source/Programs/AutomationTool/Turnkey/Commands/VerifySdk.cs
+++ b/Engine/Source/Programs/AutomationTool/Turnkey/Commands/VerifySdk.cs
@@ -215,7 +215,8 @@ namespace Turnkey.Commands
 							DeviceState |= SdkUtils.LocalAvailability.Device_CannotConnect;
 						}
 
-						TurnkeyUtils.Report("{0}@{1}: (Status={2}, Installed={3}, MinAllowed={4}, MaxAllowed={5}, Flags=\"{6}\")", Platform, Device.Name, StatusString, Device.SoftwareVersion,
+						// @todo turnkey: if Name has a comma in it, we are in trouble, and maybe a )
+						TurnkeyUtils.Report("{0}@{1}: (Name={2}, Status={3}, Installed={4}, MinAllowed={5}, MaxAllowed={6}, Flags=\"{7}\")", Platform, Device.Id, Device.Name, StatusString, Device.SoftwareVersion,
 							MinSoftwareAllowedVersion, MaxSoftwareAllowedVersion, DeviceState.ToString());
 
 						if (bForceDeviceInstall || !bIsSoftwareValid)
diff --git a/Engine/Source/Programs/UnrealBuildTool/Platform/IOS/IOSPlatformSDK.cs b/Engine/Source/Programs/UnrealBuildTool/Platform/IOS/IOSPlatformSDK.cs
new file mode 100644
--- /dev/null
+++ b/Engine/Source/Programs/UnrealBuildTool/Platform/IOS/IOSPlatformSDK.cs
@@ -0,0 +1,24 @@
+// Copyright Epic Games, Inc. All Rights Reserved.
+
+using System;
+using System.Collections.Generic;
+using System.IO;
+using EpicGames.Core;
+using System.Text.RegularExpressions;
+using Microsoft.Win32;
+using System.Diagnostics;
+
+#nullable disable
+
+namespace UnrealBuildTool
+{
+	internal class IOSPlatformSDK : ApplePlatformSDK
+	{
+		public override void GetValidSoftwareVersionRange(out string MinVersion, out string MaxVersion)
+		{
+			// what is our min IOS version?
+			MinVersion = "12.0";
+			MaxVersion = null;
+		}
+	}
+}
diff --git a/Engine/Source/Programs/UnrealBuildTool/Platform/IOS/UEBuildIOS.cs b/Engine/Source/Programs/UnrealBuildTool/Platform/IOS/UEBuildIOS.cs
--- a/Engine/Source/Programs/UnrealBuildTool/Platform/IOS/UEBuildIOS.cs
+++ b/Engine/Source/Programs/UnrealBuildTool/Platform/IOS/UEBuildIOS.cs
@@ -1156,7 +1156,7 @@ namespace UnrealBuildTool
 		/// </summary>
 		public override void RegisterBuildPlatforms()
 		{
-			ApplePlatformSDK SDK = new ApplePlatformSDK();
+			ApplePlatformSDK SDK = new IOSPlatformSDK();
 
 			// Register this build platform for IOS
 			UEBuildPlatform.RegisterBuildPlatform(new IOSPlatform(SDK));
diff --git a/Engine/Source/Programs/UnrealBuildTool/Platform/TVOS/UEBuildTVOS.cs b/Engine/Source/Programs/UnrealBuildTool/Platform/TVOS/UEBuildTVOS.cs
--- a/Engine/Source/Programs/UnrealBuildTool/Platform/TVOS/UEBuildTVOS.cs
+++ b/Engine/Source/Programs/UnrealBuildTool/Platform/TVOS/UEBuildTVOS.cs
@@ -155,7 +155,7 @@ namespace UnrealBuildTool
 		/// </summary>
 		public override void RegisterBuildPlatforms()
 		{
-			ApplePlatformSDK SDK = new ApplePlatformSDK();
+			ApplePlatformSDK SDK = new IOSPlatformSDK();
 		
 			// Register this build platform for IOS
 			UEBuildPlatform.RegisterBuildPlatform(new TVOSPlatform(SDK));
diff --git a/Engine/Source/Programs/UnrealBuildTool/System/Utils.cs b/Engine/Source/Programs/UnrealBuildTool/System/Utils.cs
--- a/Engine/Source/Programs/UnrealBuildTool/System/Utils.cs
+++ b/Engine/Source/Programs/UnrealBuildTool/System/Utils.cs
@@ -445,6 +445,7 @@ namespace UnrealBuildTool
 			StartInfo.RedirectStandardOutput = true;
 			StartInfo.RedirectStandardError = true;
 			StartInfo.CreateNoWindow = true;
+			StartInfo.StandardOutputEncoding = Encoding.UTF8;
 
 			string FullOutput = "";
 			string ErrorOutput = "";
