- Fixed crash with networked ios devices during discovery

diff --git a/Engine/Source/Developer/IOS/IOSTargetPlatform/Private/Windows/IOSDeviceHelperWindows.cpp b/Engine/Source/Developer/IOS/IOSTargetPlatform/Private/Windows/IOSDeviceHelperWindows.cpp
--- a/Engine/Source/Developer/IOS/IOSTargetPlatform/Private/Windows/IOSDeviceHelperWindows.cpp
+++ b/Engine/Source/Developer/IOS/IOSTargetPlatform/Private/Windows/IOSDeviceHelperWindows.cpp
@@ -46,6 +46,7 @@ struct FDeviceNotificationCallbackInformation
 			if (CurrentDeviceIds[Index].Contains("Network"))
 			{
 				CurrentDeviceIds.RemoveAt(Index);
+				Index--;
 				continue;
 			}
 			CurrentDeviceIds[Index].Split(TEXT(" "), &CurrentDeviceIds[Index], nullptr, ESearchCase::CaseSensitive, ESearchDir::FromStart);
