Removed PVRTC support, iOS apps will be packaged using ASTC compression (requires A8, iPhone6)

diff --git a/Engine/Config/BaseEngine.ini b/Engine/Config/BaseEngine.ini
--- a/Engine/Config/BaseEngine.ini
+++ b/Engine/Config/BaseEngine.ini
@@ -2126,8 +2126,6 @@ bSupportsLandscapeLeftOrientation=True
 bSupportsLandscapeRightOrientation=True
 PreferredLandscapeOrientation=LandscapeLeft
 bSupportsMetal=True
-bCookPVRTCTextures=True
-bCookASTCTextures=False
 bSupportsMetalMRT=False
 bShipForBitcode=True
 bGameSupportsMultipleActiveControllers=False
diff --git a/Engine/Plugins/Runtime/Interchange/Source/Nodes/Public/InterchangeTextureNode.h b/Engine/Plugins/Runtime/Interchange/Source/Nodes/Public/InterchangeTextureNode.h
--- a/Engine/Plugins/Runtime/Interchange/Source/Nodes/Public/InterchangeTextureNode.h
+++ b/Engine/Plugins/Runtime/Interchange/Source/Nodes/Public/InterchangeTextureNode.h
@@ -522,23 +522,6 @@ public:
 #endif
 	}
 
-	/** Return false if the Attribute was not set previously.*/
-	UFUNCTION(BlueprintCallable, Category = "Interchange | Node | Texture")
-	bool GetCustombForcePVRTC4(bool& AttributeValue) const
-	{
-		IMPLEMENT_NODE_ATTRIBUTE_GETTER(bForcePVRTC4, bool);
-	}
-
-	UFUNCTION(BlueprintCallable, Category = "Interchange | Node | Texture")
-	bool SetCustombForcePVRTC4(const bool& AttributeValue, bool bAddApplyDelegate = true)
-	{
-#if WITH_EDITORONLY_DATA
-		IMPLEMENT_NODE_ATTRIBUTE_SETTER(UInterchangeTextureNode, bForcePVRTC4, bool, UTexture)
-#else
-		IMPLEMENT_NODE_ATTRIBUTE_SETTER_NODELEGATE(bForcePVRTC4, bool)
-#endif
-	}
-
 	/** Return false if the Attribute was not set previously.*/
 	UFUNCTION(BlueprintCallable, Category = "Interchange | Node | Texture")
 	bool GetCustomPowerOfTwoMode(uint8& AttributeValue) const
@@ -911,7 +894,6 @@ private:
 	IMPLEMENT_NODE_ATTRIBUTE_APPLY_UOBJECT(bDitherMipMapAlpha, bool, UTexture, );
 	IMPLEMENT_NODE_ATTRIBUTE_APPLY_UOBJECT(AlphaCoverageThresholds, FVector4, UTexture, );
 	IMPLEMENT_NODE_ATTRIBUTE_APPLY_UOBJECT(bFlipGreenChannel, bool, UTexture, );
-	IMPLEMENT_NODE_ATTRIBUTE_APPLY_UOBJECT(bForcePVRTC4, bool, UTexture, );
 	IMPLEMENT_NODE_ATTRIBUTE_APPLY_UOBJECT(PowerOfTwoMode, uint8, UTexture, TEnumAsByte<enum ETexturePowerOfTwoSetting::Type>);
 	IMPLEMENT_NODE_ATTRIBUTE_APPLY_UOBJECT(PaddingColor, FColor, UTexture, );
 #endif //#if WITH_EDITORONLY_DATA
diff --git a/Engine/Source/Developer/DeveloperToolSettings/Classes/CookerSettings.h b/Engine/Source/Developer/DeveloperToolSettings/Classes/CookerSettings.h
--- a/Engine/Source/Developer/DeveloperToolSettings/Classes/CookerSettings.h
+++ b/Engine/Source/Developer/DeveloperToolSettings/Classes/CookerSettings.h
@@ -129,10 +129,6 @@ public:
 	UPROPERTY(GlobalConfig, EditAnywhere, Category = Cooker, AdvancedDisplay, meta = (DisplayName = "r values that need to be versioned"))
 	TArray<FString> VersionedIntRValues;
 
-	/** Quality of 0 means fastest, 4 means best quality */
-	UPROPERTY(GlobalConfig, EditAnywhere, Category = Textures, meta = (DisplayName = "PVRTC Compression Quality (0-4, 0 is fastest)"))
-	int32 DefaultPVRTCQuality;
-
 	/** Quality of 0 means fastest, 3 means best quality */
 	UPROPERTY(GlobalConfig, EditAnywhere, Category = Textures, meta = (DisplayName = "ASTC Compression Quality vs Speed (0-3, 0 is fastest)"))
 	int32 DefaultASTCQualityBySpeed;
diff --git a/Engine/Source/Developer/DeveloperToolSettings/Private/CookerSettings.cpp b/Engine/Source/Developer/DeveloperToolSettings/Private/CookerSettings.cpp
--- a/Engine/Source/Developer/DeveloperToolSettings/Private/CookerSettings.cpp
+++ b/Engine/Source/Developer/DeveloperToolSettings/Private/CookerSettings.cpp
@@ -20,7 +20,6 @@ UCookerSettings::UCookerSettings(const FObjectInitializer& ObjectInitializer)
 	, bCookBlueprintComponentTemplateData(false)
 {
 	SectionName = TEXT("Cooker");
-	DefaultPVRTCQuality = 1;
 	DefaultASTCQualityBySize = 3;
 	DefaultASTCQualityBySpeed = 3;
 }
diff --git a/Engine/Source/Developer/IOS/IOSTargetPlatform/Private/IOSTargetPlatform.cpp b/Engine/Source/Developer/IOS/IOSTargetPlatform/Private/IOSTargetPlatform.cpp
--- a/Engine/Source/Developer/IOS/IOSTargetPlatform/Private/IOSTargetPlatform.cpp
+++ b/Engine/Source/Developer/IOS/IOSTargetPlatform/Private/IOSTargetPlatform.cpp
@@ -437,22 +437,6 @@ static bool SupportsMetalMRT()
 	return bSupportsMetalMRT;
 }
 
-static bool CookPVRTC()
-{
-	// default to using PVRTC
-	bool bCookPVRTCTextures = true;
-	GConfig->GetBool(TEXT("/Script/IOSRuntimeSettings.IOSRuntimeSettings"), TEXT("bCookPVRTCTextures"), bCookPVRTCTextures, GEngineIni);
-	return bCookPVRTCTextures;
-}
-
-static bool CookASTC()
-{
-	// default to not using ASTC
-	bool bCookASTCTextures = true;
-	GConfig->GetBool(TEXT("/Script/IOSRuntimeSettings.IOSRuntimeSettings"), TEXT("bCookASTCTextures"), bCookASTCTextures, GEngineIni);
-	return bCookASTCTextures;
-}
-
 static bool SupportsSoftwareOcclusion()
 {
 	static auto* CVarMobileAllowSoftwareOcclusion = IConsoleManager::Get().FindTConsoleVariableDataInt(TEXT("r.Mobile.AllowSoftwareOcclusion"));
@@ -562,18 +546,18 @@ void FIOSTargetPlatform::GetAllTargetedShaderFormats( TArray<FName>& OutFormats
 	GetAllPossibleShaderFormats(OutFormats);
 }
 
-// we remap some of the defaults (with PVRTC and ASTC formats)
+// we remap some of the defaults
 static FName FormatRemap[] =
 {
-	// original				PVRTC						ASTC
-	FName(TEXT("DXT1")),	FName(TEXT("PVRTC2")),		FName(TEXT("ASTC_RGB")),
-	FName(TEXT("DXT5")),	FName(TEXT("PVRTC4")),		FName(TEXT("ASTC_RGBA")),
-	FName(TEXT("DXT5n")),	FName(TEXT("PVRTCN")),		FName(TEXT("ASTC_NormalAG")),
-	FName(TEXT("BC5")),		FName(TEXT("PVRTCN")),		FName(TEXT("ASTC_NormalRG")),
-	FName(TEXT("AutoDXT")),	FName(TEXT("AutoPVRTC")),	FName(TEXT("ASTC_RGBAuto")),
-	FName(TEXT("BC4")),		FName(TEXT("G8")),			FName(TEXT("G8")),
-	FName(TEXT("BC6H")),	FName(TEXT("PVRTC2")),		FName(TEXT("ASTC_RGB")), 
-	FName(TEXT("BC7")),		FName(TEXT("AutoPVRTC")),	FName(TEXT("ASTC_RGBAuto"))
+	// original				ASTC
+	FName(TEXT("AutoDXT")),	FName(TEXT("ASTC_RGBAuto")),
+	FName(TEXT("DXT1")),	FName(TEXT("ASTC_RGB")),
+	FName(TEXT("DXT5")),	FName(TEXT("ASTC_RGBA")),
+	FName(TEXT("DXT5n")),	FName(TEXT("ASTC_NormalAG")),
+	FName(TEXT("BC5")),		FName(TEXT("ASTC_NormalRG")),
+	FName(TEXT("BC4")),		FName(TEXT("G8")),
+	FName(TEXT("BC6H")),	FName(TEXT("ASTC_RGB")), 
+	FName(TEXT("BC7")),		FName(TEXT("ASTC_RGBAuto"))
 };
 static FName NameBGRA8(TEXT("BGRA8"));
 static FName NameG8 = FName(TEXT("G8"));
@@ -582,23 +566,10 @@ void FIOSTargetPlatform::GetTextureFormats( const UTexture* Texture, TArray< TAr
 {
 	check(Texture);
 
-	static FName NamePOTERROR(TEXT("POTERROR"));
-
 	const int32 NumLayers = Texture->Source.GetNumLayers();
-
-	if (Texture->bForcePVRTC4 && CookPVRTC())
-	{
-		TArray<FName> NamesPVRTC4;
-		TArray<FName> NamesPVRTCN;
-		NamesPVRTC4.Init(FName(TEXT("PVRTC4")), NumLayers);
-		NamesPVRTCN.Init(FName(TEXT("PVRTCN")), NumLayers);
-
-		OutFormats.AddUnique(NamesPVRTC4);
-		OutFormats.AddUnique(NamesPVRTCN);
-		return;
-	}
-
-	TArray<FName> TextureFormatNames;
+	
+	TArray<FName>& TextureFormatNames = OutFormats.AddDefaulted_GetRef();
+	TextureFormatNames.Reserve(NumLayers);
 
 	// forward rendering only needs one channel for shadow maps
 	if (Texture->LODGroup == TEXTUREGROUP_Shadowmap && !SupportsMetalMRT())
@@ -607,64 +578,31 @@ void FIOSTargetPlatform::GetTextureFormats( const UTexture* Texture, TArray< TAr
 	}
 
 	// if we didn't assign anything specially, then use the defaults
-    bool bIncludePVRTC = !bIsTVOS && CookPVRTC();
-    bool bIncludeASTC = bIsTVOS || CookASTC();
 	if (TextureFormatNames.Num() == 0)
 	{
-        int32 BlockSize = 4;
-        if (!Texture->bForcePVRTC4 && !bIncludePVRTC && bIncludeASTC)
-        {
-            BlockSize = 1;
-        }
+        int32 BlockSize = 1;
 		GetDefaultTextureFormatNamePerLayer(TextureFormatNames, this, Texture, true, false, BlockSize);
 	}
 
-	// include the formats we want (use ASTC first so that it is preferred at runtime if they both exist and it's supported)
-	if (bIncludeASTC)
+	// include the formats we want
+	for (FName& TextureFormatName : TextureFormatNames)
 	{
-		TArray<FName> TextureFormatNamesASTC(TextureFormatNames);
-		for (FName& TextureFormatName : TextureFormatNamesASTC)
+		if (TextureFormatName == NameBGRA8 || 
+			TextureFormatName == NameG8)
 		{
-			for (int32 RemapIndex = 0; RemapIndex < UE_ARRAY_COUNT(FormatRemap); RemapIndex += 3)
-			{
-				if (TextureFormatName == FormatRemap[RemapIndex])
-				{
-					TextureFormatName = FormatRemap[RemapIndex + 2];
-					break;
-				}
-			}
+			continue;
 		}
-		OutFormats.AddUnique(TextureFormatNamesASTC);
-	}
 
-	if (bIncludePVRTC)
-	{
-		TArray<FName> TextureFormatNamesPVRTC(TextureFormatNames);
-		for (FName& TextureFormatName : TextureFormatNamesPVRTC)
+		for (int32 RemapIndex = 0; RemapIndex < UE_ARRAY_COUNT(FormatRemap); RemapIndex += 2)
 		{
-			for (int32 RemapIndex = 0; RemapIndex < UE_ARRAY_COUNT(FormatRemap); RemapIndex += 3)
+			if (TextureFormatName == FormatRemap[RemapIndex])
 			{
-				if (TextureFormatName == FormatRemap[RemapIndex])
-				{
-					// handle non-power of 2 textures
-					if (!Texture->Source.IsPowerOfTwo() && Texture->PowerOfTwoMode == ETexturePowerOfTwoSetting::None)
-					{
-						// option 1: Uncompress, but users will get very large textures unknowingly
-						// TextureFormatName = NameBGRA8;
-						// option 2: Use an "error message" texture so they see it in game
-						TextureFormatName = NamePOTERROR;
-					}
-					else
-					{
-						TextureFormatName = FormatRemap[RemapIndex + 1];
-					}
-					break;
-				}
+				TextureFormatName = FormatRemap[RemapIndex + 1];
+				break;
 			}
 		}
-		OutFormats.AddUnique(TextureFormatNamesPVRTC);
 	}
-
+		
 	for (FName& TextureFormatName : OutFormats.Last())
 	{
 		if (Texture->IsA(UTextureCube::StaticClass()))
@@ -686,30 +624,18 @@ void FIOSTargetPlatform::GetTextureFormats( const UTexture* Texture, TArray< TAr
 void FIOSTargetPlatform::GetAllTextureFormats(TArray<FName>& OutFormats) const 
 {
 	bool bFoundRemap = false;
-	bool bIncludePVRTC = !bIsTVOS && CookPVRTC();
-	bool bIncludeASTC = bIsTVOS || CookASTC();
 
 	GetAllDefaultTextureFormats(this, OutFormats, false);
 
-	for (int32 RemapIndex = 0; RemapIndex < UE_ARRAY_COUNT(FormatRemap); RemapIndex += 3)
+	for (int32 RemapIndex = 0; RemapIndex < UE_ARRAY_COUNT(FormatRemap); RemapIndex += 2)
 	{
 		OutFormats.Remove(FormatRemap[RemapIndex+0]);
 	}
 
-	// include the formats we want (use ASTC first so that it is preferred at runtime if they both exist and it's supported)
-	if (bIncludeASTC)
+	// include the formats we want
+	for (int32 RemapIndex = 0; RemapIndex < UE_ARRAY_COUNT(FormatRemap); RemapIndex += 2)
 	{
-		for (int32 RemapIndex = 0; RemapIndex < UE_ARRAY_COUNT(FormatRemap); RemapIndex += 3)
-		{
-			OutFormats.AddUnique(FormatRemap[RemapIndex + 2]);
-		}
-	}
-	if (bIncludePVRTC)
-	{
-		for (int32 RemapIndex = 0; RemapIndex < UE_ARRAY_COUNT(FormatRemap); RemapIndex += 3)
-		{
-			OutFormats.AddUnique(FormatRemap[RemapIndex + 1]);
-		}
+		OutFormats.AddUnique(FormatRemap[RemapIndex + 1]);
 	}
 }
 
@@ -728,10 +654,6 @@ FName FIOSTargetPlatform::FinalizeVirtualTextureLayerFormat(FName Format) const
 		{ { FName(TEXT("ASTC_RGBAuto")) },		{ NameAutoETC2 } },
 		{ { FName(TEXT("ASTC_NormalAG")) },		{ NameETC2_RGB } },
 		{ { FName(TEXT("ASTC_NormalRG")) },		{ NameETC2_RGB } },
-		{ { FName(TEXT("PVRTC2")) },			{ NameETC2_RGB } },
-		{ { FName(TEXT("PVRTC4")) },			{ NameETC2_RGBA } },
-		{ { FName(TEXT("PVRTCN")) },			{ NameETC2_RGB } },
-		{ { FName(TEXT("AutoPVRTC")) },			{ NameAutoETC2 } }
 	};
 
 	for (int32 RemapIndex = 0; RemapIndex < UE_ARRAY_COUNT(ETCRemap); RemapIndex++)
diff --git a/Engine/Source/Developer/TargetPlatform/TargetPlatform.Build.cs b/Engine/Source/Developer/TargetPlatform/TargetPlatform.Build.cs
--- a/Engine/Source/Developer/TargetPlatform/TargetPlatform.Build.cs
+++ b/Engine/Source/Developer/TargetPlatform/TargetPlatform.Build.cs
@@ -47,7 +47,6 @@ public class TargetPlatform : ModuleRules
 				DynamicallyLoadedModuleNames.Add("MetalShaderFormat");
 
 				DynamicallyLoadedModuleNames.Add("TextureFormatDXT");
-				DynamicallyLoadedModuleNames.Add("TextureFormatPVR");
 				DynamicallyLoadedModuleNames.Add("TextureFormatASTC");
 
 				DynamicallyLoadedModuleNames.Add("TextureFormatUncompressed");
@@ -70,7 +69,6 @@ public class TargetPlatform : ModuleRules
 			else if (Target.Platform == UnrealTargetPlatform.Mac)
 			{
 				DynamicallyLoadedModuleNames.Add("TextureFormatDXT");
-				DynamicallyLoadedModuleNames.Add("TextureFormatPVR");
 				DynamicallyLoadedModuleNames.Add("TextureFormatASTC");
 
 				DynamicallyLoadedModuleNames.Add("TextureFormatUncompressed");
@@ -92,7 +90,6 @@ public class TargetPlatform : ModuleRules
 			else if (Target.IsInPlatformGroup(UnrealPlatformGroup.Linux))
 			{
 				DynamicallyLoadedModuleNames.Add("TextureFormatDXT");
-				DynamicallyLoadedModuleNames.Add("TextureFormatPVR");
 				DynamicallyLoadedModuleNames.Add("TextureFormatASTC");
 
 				DynamicallyLoadedModuleNames.Add("TextureFormatUncompressed");
diff --git a/Engine/Source/Developer/TextureFormatPVR/Private/TextureFormatPVR.cpp b/Engine/Source/Developer/TextureFormatPVR/Private/TextureFormatPVR.cpp
deleted file mode 100644
diff --git a/Engine/Source/Developer/TextureFormatPVR/TextureFormatPVR.Build.cs b/Engine/Source/Developer/TextureFormatPVR/TextureFormatPVR.Build.cs
deleted file mode 100644
diff --git a/Engine/Source/Programs/UnrealBuildTool/Platform/IOS/UEBuildIOS.cs b/Engine/Source/Programs/UnrealBuildTool/Platform/IOS/UEBuildIOS.cs
--- a/Engine/Source/Programs/UnrealBuildTool/Platform/IOS/UEBuildIOS.cs
+++ b/Engine/Source/Programs/UnrealBuildTool/Platform/IOS/UEBuildIOS.cs
@@ -977,7 +977,6 @@ namespace UnrealBuildTool
 					else if (ModuleName == "TargetPlatform")
 					{
 						bBuildShaderFormats = true;
-						Rules.DynamicallyLoadedModuleNames.Add("TextureFormatPVR");
 						Rules.DynamicallyLoadedModuleNames.Add("TextureFormatASTC");
 						Rules.DynamicallyLoadedModuleNames.Add("TextureFormatETC2");
 						if (Target.bBuildDeveloperTools && Target.bCompileAgainstEngine)
diff --git a/Engine/Source/Runtime/Apple/MetalRHI/Private/MetalRHI.cpp b/Engine/Source/Runtime/Apple/MetalRHI/Private/MetalRHI.cpp
--- a/Engine/Source/Runtime/Apple/MetalRHI/Private/MetalRHI.cpp
+++ b/Engine/Source/Runtime/Apple/MetalRHI/Private/MetalRHI.cpp
@@ -195,13 +195,8 @@ FMetalDynamicRHI::FMetalDynamicRHI(ERHIFeatureLevel::Type RequestedFeatureLevel)
 #if PLATFORM_IOS
 	// A8 can use 256 bits of MRTs
 #if PLATFORM_TVOS
-	bool bCanUseWideMRTs = true;
-	bool bCanUseASTC = true;
 	GRHISupportsDrawIndirect = Device.SupportsFeatureSet(mtlpp::FeatureSet::tvOS_GPUFamily2_v1);
 #else
-	bool bCanUseWideMRTs = Device.SupportsFeatureSet(mtlpp::FeatureSet::iOS_GPUFamily2_v1);
-	bool bCanUseASTC = Device.SupportsFeatureSet(mtlpp::FeatureSet::iOS_GPUFamily2_v1) && !FParse::Param(FCommandLine::Get(),TEXT("noastc"));
-	
 	GRHISupportsRWTextureBuffers = Device.SupportsFeatureSet(mtlpp::FeatureSet::iOS_GPUFamily4_v1);
 	GRHISupportsDrawIndirect = Device.SupportsFeatureSet(mtlpp::FeatureSet::iOS_GPUFamily3_v1);
 
@@ -239,7 +234,7 @@ FMetalDynamicRHI::FMetalDynamicRHI(ERHIFeatureLevel::Type RequestedFeatureLevel)
 	bSupportsRHIThread = FParse::Param(FCommandLine::Get(),TEXT("rhithread"));
 
     // only allow GBuffers, etc on A8s (A7s are just not going to cut it)
-    if (bProjectSupportsMRTs && bCanUseWideMRTs && bRequestedMetalMRT)
+    if (bProjectSupportsMRTs && bRequestedMetalMRT)
     {
 #if PLATFORM_TVOS
 		ValidateTargetedRHIFeatureLevelExists(SP_METAL_MRT);
@@ -292,9 +287,6 @@ FMetalDynamicRHI::FMetalDynamicRHI(ERHIFeatureLevel::Type RequestedFeatureLevel)
 	check(DeviceIndex < GPUs.Num());
 	FMacPlatformMisc::FGPUDescriptor const& GPUDesc = GPUs[DeviceIndex];
 	
-    // A8 can use 256 bits of MRTs
-    bool bCanUseWideMRTs = true;
-    bool bCanUseASTC = false;
 	bool bSupportsD24S8 = false;
 	bool bSupportsD16 = false;
 	
@@ -572,7 +564,7 @@ FMetalDynamicRHI::FMetalDynamicRHI(ERHIFeatureLevel::Type RequestedFeatureLevel)
 	GHardwareHiddenSurfaceRemoval = true;
 	GSupportsRenderTargetFormat_PF_G8 = false;
 	GRHISupportsTextureStreaming = true;
-	GSupportsWideMRT = bCanUseWideMRTs;
+	GSupportsWideMRT = true;
 	// GSupportsTransientResourceAliasing = FMetalCommandQueue::SupportsFeature(EMetalFeaturesHeaps) && FMetalCommandQueue::SupportsFeature(EMetalFeaturesFences);
 	GSupportsSeparateRenderTargetBlendState = (GMaxRHIFeatureLevel >= ERHIFeatureLevel::SM5);
 
@@ -739,15 +731,15 @@ FMetalDynamicRHI::FMetalDynamicRHI(ERHIFeatureLevel::Type RequestedFeatureLevel)
 	GPixelFormats[PF_PVRTC4				].PlatformFormat	= (uint32)mtlpp::PixelFormat::PVRTC_RGBA_4BPP;
 	GPixelFormats[PF_PVRTC4				].Supported			= true;
 	GPixelFormats[PF_ASTC_4x4			].PlatformFormat	= (uint32)mtlpp::PixelFormat::ASTC_4x4_LDR;
-	GPixelFormats[PF_ASTC_4x4			].Supported			= bCanUseASTC;
+	GPixelFormats[PF_ASTC_4x4			].Supported			= true;
 	GPixelFormats[PF_ASTC_6x6			].PlatformFormat	= (uint32)mtlpp::PixelFormat::ASTC_6x6_LDR;
-	GPixelFormats[PF_ASTC_6x6			].Supported			= bCanUseASTC;
+	GPixelFormats[PF_ASTC_6x6			].Supported			= true;
 	GPixelFormats[PF_ASTC_8x8			].PlatformFormat	= (uint32)mtlpp::PixelFormat::ASTC_8x8_LDR;
-	GPixelFormats[PF_ASTC_8x8			].Supported			= bCanUseASTC;
+	GPixelFormats[PF_ASTC_8x8			].Supported			= true;
 	GPixelFormats[PF_ASTC_10x10			].PlatformFormat	= (uint32)mtlpp::PixelFormat::ASTC_10x10_LDR;
-	GPixelFormats[PF_ASTC_10x10			].Supported			= bCanUseASTC;
+	GPixelFormats[PF_ASTC_10x10			].Supported			= true;
 	GPixelFormats[PF_ASTC_12x12			].PlatformFormat	= (uint32)mtlpp::PixelFormat::ASTC_12x12_LDR;
-	GPixelFormats[PF_ASTC_12x12			].Supported			= bCanUseASTC;
+	GPixelFormats[PF_ASTC_12x12			].Supported			= true;
 
 	// used with virtual textures
 	GPixelFormats[PF_ETC2_RGB	  		].PlatformFormat	= (uint32)mtlpp::PixelFormat::ETC2_RGB8;
diff --git a/Engine/Source/Runtime/Engine/Classes/Engine/Texture.h b/Engine/Source/Runtime/Engine/Classes/Engine/Texture.h
--- a/Engine/Source/Runtime/Engine/Classes/Engine/Texture.h
+++ b/Engine/Source/Runtime/Engine/Classes/Engine/Texture.h
@@ -739,10 +739,6 @@ public:
 	UPROPERTY(EditAnywhere, BlueprintReadWrite, Category=Texture, AdvancedDisplay)
 	uint32 bFlipGreenChannel:1;
 
-	/** For DXT1 textures, setting this will cause the texture to be twice the size, but better looking, on iPhone */
-	UPROPERTY(EditAnywhere, Category = Texture, AdvancedDisplay)
-	uint32 bForcePVRTC4:1;
-
 	/** How to pad the texture to a power of 2 size (if necessary) */
 	UPROPERTY(EditAnywhere, BlueprintReadWrite, Category=Texture)
 	TEnumAsByte<enum ETexturePowerOfTwoSetting::Type> PowerOfTwoMode;
diff --git a/Engine/Source/Runtime/Engine/Private/LightMap.cpp b/Engine/Source/Runtime/Engine/Private/LightMap.cpp
--- a/Engine/Source/Runtime/Engine/Private/LightMap.cpp
+++ b/Engine/Source/Runtime/Engine/Private/LightMap.cpp
@@ -1559,7 +1559,6 @@ void FLightMapPendingTexture::EncodeCoefficientTexture(int32 CoefficientIndex, U
 	FormatSettings.SRGB = false;
 	FormatSettings.CompressionNoAlpha = CoefficientIndex >= LQ_LIGHTMAP_COEF_INDEX;
 	FormatSettings.CompressionNone = !GCompressLightmaps;
-	//FormatSettings.bForcePVRTC4 = true;
 	Texture->SetLayerFormatSettings(LayerIndex, FormatSettings);
 
 	if (bEncodeVirtualTexture)
@@ -1822,7 +1821,6 @@ void FLightMapPendingTexture::StartEncoding(ULevel* LightingScenario, ITextureCo
 		Texture->Filter	= GUseBilinearLightmaps ? TF_Default : TF_Nearest;
 		Texture->LODGroup = TEXTUREGROUP_Lightmap;
 		Texture->LightmapFlags = ELightMapFlags( LightmapFlags );
-		Texture->bForcePVRTC4 = true;
 
 		EncodeCoefficientTexture(CoefficientIndex, Texture, 0u, TextureColor, false);
 	}
diff --git a/Engine/Source/Runtime/IOS/IOSRuntimeSettings/Classes/IOSRuntimeSettings.h b/Engine/Source/Runtime/IOS/IOSRuntimeSettings/Classes/IOSRuntimeSettings.h
--- a/Engine/Source/Runtime/IOS/IOSRuntimeSettings/Classes/IOSRuntimeSettings.h
+++ b/Engine/Source/Runtime/IOS/IOSRuntimeSettings/Classes/IOSRuntimeSettings.h
@@ -219,14 +219,6 @@ public:
 	UPROPERTY(GlobalConfig, EditAnywhere, Category = Rendering, meta = (DisplayName = "Metal Desktop Renderer"))
 	bool bSupportsMetalMRT;
 	
-	// Whether or not to add support for PVRTC textures
-	UPROPERTY(GlobalConfig, EditAnywhere, Category = Rendering, meta = (DisplayName = "Cook PVRTC texture data for OpenGL ES or Metal on A7 and earlier devices"))
-	bool bCookPVRTCTextures;
-
-	// Whether or not to add support for ASTC textures
-	UPROPERTY(GlobalConfig, EditAnywhere, Category = Rendering, meta = (DisplayName = "Cook ASTC texture data for Metal on A8 or later devices"))
-	bool bCookASTCTextures;
-	
     // Whether to build the iOS project as a framework.
     UPROPERTY(GlobalConfig, EditAnywhere, Category = Build, meta = (DisplayName = "Build project as a framework (Experimental)"))
     bool bBuildAsFramework;
