Add MacOS support for libimobile interface in editor and turnkey

diff --git a/Engine/Source/Developer/IOS/IOSTargetPlatform/Private/IOSDeviceHelper.cpp b/Engine/Source/Developer/IOS/IOSTargetPlatform/Private/IOSDeviceHelper.cpp
new file mode 100644
--- /dev/null
+++ b/Engine/Source/Developer/IOS/IOSTargetPlatform/Private/IOSDeviceHelper.cpp
@@ -0,0 +1,419 @@
+#include "IOSTargetPlatform.h"
+#include "IOSTargetDeviceOutput.h"
+#include "HAL/PlatformProcess.h"
+#include "HAL/Runnable.h"
+#include "HAL/RunnableThread.h"
+#include "Interfaces/ITargetPlatformManagerModule.h"
+#include "Interfaces/ITargetPlatform.h"
+#include "Interfaces/IProjectManager.h"
+
+DEFINE_LOG_CATEGORY_STATIC(LogIOSDeviceHelper, Log, All);
+
+struct FDeviceNotificationCallbackInformation
+{
+    FString UDID;
+    FString    DeviceName;
+    FString ProductType;
+    uint32 msgType;
+};
+
+
+    struct LibIMobileDevice
+    {
+        FString DeviceName;
+        FString DeviceID;
+        FString DeviceType;
+    };
+
+    static FString GetLibImobileDeviceExe(const FString &ExeName)
+    {
+        FString ToReturn;
+#if PLATFORM_WINDOWS
+        ToReturn = FPaths::ConvertRelativePathToFull(FPaths::EngineDir() / TEXT("Extras/ThirdPartyNotUE/libimobiledevice/x64/"));
+#elif PLATFORM_MAC
+        ToReturn = FPaths::ConvertRelativePathToFull(FPaths::EngineDir() / TEXT("Extras/ThirdPartyNotUE/libimobiledevice/Mac/"));
+#else
+        UE_LOG(LogIOSDeviceHelper, Error, TEXT("The current platform is unsupported by Libimobile library."));
+#endif
+        ToReturn += ExeName;
+        return ToReturn;
+    }
+
+    static TArray<LibIMobileDevice> GetLibIMobileDevices()
+    {
+        FString OutStdOut;
+        FString OutStdErr;
+        FString LibimobileDeviceId = GetLibImobileDeviceExe("idevice_id");
+        int ReturnCode;
+        // get the list of devices UDID
+        FPlatformProcess::ExecProcess(*LibimobileDeviceId, TEXT(""), &ReturnCode, &OutStdOut, &OutStdErr);
+        
+
+        TArray<LibIMobileDevice> ToReturn;
+
+        // separate out each line
+        TArray<FString> CurrentDeviceIds;
+        OutStdOut.ParseIntoArray(CurrentDeviceIds, TEXT("\n"), true);
+        for (int32 Index = 0; Index != CurrentDeviceIds.Num(); ++Index)
+        {
+            if (CurrentDeviceIds[Index].Contains("Network"))
+            {
+                CurrentDeviceIds.RemoveAt(Index);
+                Index--;
+                continue;
+            }
+            CurrentDeviceIds[Index].Split(TEXT(" "), &CurrentDeviceIds[Index], nullptr, ESearchCase::CaseSensitive, ESearchDir::FromStart);
+        }
+        TArray<FString> DeviceStrings;
+        for (int32 StringIndex = 0; StringIndex < CurrentDeviceIds.Num(); ++StringIndex)
+        {
+
+            const FString& DeviceID = CurrentDeviceIds[StringIndex];
+
+            FString OutStdOutInfo;
+            FString OutStdErrInfo;
+            FString LibimobileDeviceInfo = GetLibImobileDeviceExe("ideviceinfo");
+            int ReturnCodeInfo;
+            FString Arguments = "-u " + DeviceID;
+            // get the list of devices UDID
+            FPlatformProcess::ExecProcess(*LibimobileDeviceInfo, *Arguments, &ReturnCodeInfo, &OutStdOutInfo, &OutStdErrInfo);
+
+            // parse product type and device name
+            FString DeviceName;
+            OutStdOutInfo.Split(TEXT("DeviceName: "), nullptr, &DeviceName, ESearchCase::CaseSensitive, ESearchDir::FromStart);
+            DeviceName.Split(LINE_TERMINATOR, &DeviceName, nullptr, ESearchCase::CaseSensitive, ESearchDir::FromStart);
+            FString ProductType;
+            OutStdOutInfo.Split(TEXT("ProductType: "), nullptr, &ProductType, ESearchCase::CaseSensitive, ESearchDir::FromStart);
+            DeviceName.Split(LINE_TERMINATOR, &DeviceName, nullptr, ESearchCase::CaseSensitive, ESearchDir::FromStart);
+            LibIMobileDevice ToAdd;
+            ToAdd.DeviceID = DeviceID;
+            ToAdd.DeviceName = DeviceName;
+            ToAdd.DeviceType = ProductType;
+            ToReturn.Add(ToAdd);
+        }
+        return ToReturn;
+    }
+
+class FIOSDevice
+{
+public:
+    FIOSDevice(FString InID, FString InName)
+        : UDID(InID)
+        , Name(InName)
+    {
+    }
+    
+    ~FIOSDevice()
+    {
+    }
+
+    FString SerialNumber() const
+    {
+        return UDID;
+    }
+
+private:
+    FString UDID;
+    FString Name;
+};
+
+/**
+ * Delegate type for devices being connected or disconnected from the machine
+ *
+ * The first parameter is newly added or removed device
+ */
+DECLARE_MULTICAST_DELEGATE_OneParam(FDeviceNotification, void*)
+
+// recheck once per minute
+#define        RECHECK_COUNTER_RESET            12
+
+class FDeviceQueryTask
+    : public FRunnable
+{
+public:
+    FDeviceQueryTask()
+        : Stopping(false)
+        , bCheckDevices(true)
+        , NeedSDKCheck(true)
+        , RetryQuery(5)
+    {}
+
+    virtual bool Init() override
+    {
+        return true;
+    }
+
+    virtual uint32 Run() override
+    {
+        while (!Stopping)
+        {
+            if (IsEngineExitRequested())
+            {
+                break;
+            }
+            if (GetTargetPlatformManager())
+            {
+                FString OutTutorialPath;
+                const ITargetPlatform* Platform = GetTargetPlatformManager()->FindTargetPlatform(TEXT("IOS"));
+                if (Platform)
+                {
+                    if (Platform->IsSdkInstalled(false, OutTutorialPath))
+                    {
+                        break;
+                    }
+                }
+                Enable(false);
+                return 0;
+            }
+            else
+            {
+                FPlatformProcess::Sleep(1.0f);
+            }
+        }
+        int RecheckCounter = RECHECK_COUNTER_RESET;
+        while (!Stopping)
+        {
+            if (IsEngineExitRequested())
+            {
+                break;
+            }
+            if (bCheckDevices)
+            {
+#if WITH_EDITOR
+                if (!IsRunningCommandlet())
+                {
+                    //if (NeedSDKCheck)
+                    //{
+                    //    NeedSDKCheck = false;
+                    //    FProjectStatus ProjectStatus;
+                    //    if (!IProjectManager::Get().QueryStatusForCurrentProject(ProjectStatus) || (!ProjectStatus.IsTargetPlatformSupported(TEXT("IOS")) && !ProjectStatus.IsTargetPlatformSupported(TEXT("TVOS"))))
+                    //    {
+                    //        Enable(false);
+                    //    }
+                    //}
+                    //else
+                    {
+                        // BHP - Turning off device check to prevent it from interfering with packaging
+                        QueryDevices();
+                    }
+                }
+#endif
+            }
+            RecheckCounter--;
+            if (RecheckCounter < 0)
+            {
+                RecheckCounter = RECHECK_COUNTER_RESET;
+                bCheckDevices = true;
+                NeedSDKCheck = true;
+            }
+
+            FPlatformProcess::Sleep(5.0f);
+        }
+
+        return 0;
+    }
+
+    virtual void Stop() override
+    {
+        Stopping = true;
+    }
+
+    virtual void Exit() override
+    {}
+
+    FDeviceNotification& OnDeviceNotification()
+    {
+        return DeviceNotification;
+    }
+
+    void Enable(bool bInCheckDevices)
+    {
+        bCheckDevices = bInCheckDevices;
+    }
+
+private:
+
+    void QueryDevices()
+    {
+        FString OutStdOut;
+        FString OutStdErr;
+        FString LibimobileDeviceId = GetLibImobileDeviceExe("idevice_id");
+        int ReturnCode;
+        // get the list of devices UDID
+        FPlatformProcess::ExecProcess(*LibimobileDeviceId, TEXT(""), &ReturnCode, &OutStdOut, &OutStdErr);
+        if (OutStdOut.Len() == 0)
+        {
+            RetryQuery--;
+            if (RetryQuery < 0)
+            {
+                UE_LOG(LogIOSDeviceHelper, Verbose, TEXT("IOS device listing is disabled for 1 minute (too many failed attempts)!"));
+                Enable(false);
+            }
+            return;
+        }
+        RetryQuery = 5;
+    
+        TArray<LibIMobileDevice> ParsedDevices = GetLibIMobileDevices();
+
+        for (int32 Index = 0; Index < ParsedDevices.Num(); ++Index)
+        {
+            // move on to next device if this one is already a known device
+            if (ConnectedDeviceIds.Find(ParsedDevices[Index].DeviceID) != INDEX_NONE)
+            {
+                ConnectedDeviceIds.Remove(ParsedDevices[Index].DeviceID);
+                continue;
+            }
+
+            // create an FIOSDevice
+            FDeviceNotificationCallbackInformation CallbackInfo;
+            CallbackInfo.DeviceName = ParsedDevices[Index].DeviceName;
+            CallbackInfo.UDID = ParsedDevices[Index].DeviceID;
+            CallbackInfo.ProductType = ParsedDevices[Index].DeviceType;
+            CallbackInfo.msgType = 1;
+            DeviceNotification.Broadcast(&CallbackInfo);
+        }
+        
+        // remove all devices no longer found
+        for (int32 DeviceIndex = 0; DeviceIndex < ConnectedDeviceIds.Num(); ++DeviceIndex)
+        {
+            FDeviceNotificationCallbackInformation CallbackInfo;
+            CallbackInfo.UDID = ConnectedDeviceIds[DeviceIndex];
+            CallbackInfo.msgType = 2;
+            DeviceNotification.Broadcast(&CallbackInfo);
+        }
+        ConnectedDeviceIds.Empty();
+        for (int32 Index = 0; Index < ParsedDevices.Num(); ++Index)
+        {
+            ConnectedDeviceIds.Add(ParsedDevices[Index].DeviceID);
+        }
+    }
+
+    bool Stopping;
+    bool bCheckDevices;
+    bool NeedSDKCheck;
+    int RetryQuery;
+    TArray<FString> ConnectedDeviceIds;
+    FDeviceNotification DeviceNotification;
+};
+
+/* FIOSDeviceHelper structors
+ *****************************************************************************/
+static TMap<FIOSDevice*, FIOSLaunchDaemonPong> ConnectedDevices;
+static FDeviceQueryTask* QueryTask = NULL;
+static FRunnableThread* QueryThread = NULL;
+static TArray<FDeviceNotificationCallbackInformation> NotificationMessages;
+static FTickerDelegate TickDelegate;
+
+bool FIOSDeviceHelper::MessageTickDelegate(float DeltaTime)
+{
+    QUICK_SCOPE_CYCLE_COUNTER(STAT_FIOSDeviceHelper_MessageTickDelegate);
+
+    for (int Index = 0; Index < NotificationMessages.Num(); ++Index)
+    {
+        FDeviceNotificationCallbackInformation cbi = NotificationMessages[Index];
+        FIOSDeviceHelper::DeviceCallback(&cbi);
+    }
+    NotificationMessages.Empty();
+
+    return true;
+}
+
+void FIOSDeviceHelper::Initialize(bool bIsTVOS)
+{
+    if(!bIsTVOS)
+    {
+        // add the message pump
+        TickDelegate = FTickerDelegate::CreateStatic(MessageTickDelegate);
+        FTicker::GetCoreTicker().AddTicker(TickDelegate, 5.0f);
+
+        // kick off a thread to query for connected devices
+        QueryTask = new FDeviceQueryTask();
+        QueryTask->OnDeviceNotification().AddStatic(FIOSDeviceHelper::DeviceCallback);
+
+        static int32 QueryTaskCount = 1;
+        if (QueryTaskCount == 1)
+        {
+            // create the socket subsystem (loadmodule in game thread)
+            ISocketSubsystem* SSS = ISocketSubsystem::Get();
+            QueryThread = FRunnableThread::Create(QueryTask, *FString::Printf(TEXT("FIOSDeviceHelper.QueryTask_%d"), QueryTaskCount++), 128 * 1024, TPri_Normal);
+        }
+    }
+}
+
+void FIOSDeviceHelper::DeviceCallback(void* CallbackInfo)
+{
+    struct FDeviceNotificationCallbackInformation* cbi = (FDeviceNotificationCallbackInformation*)CallbackInfo;
+
+    if (!IsInGameThread())
+    {
+        NotificationMessages.Add(*cbi);
+    }
+    else
+    {
+        switch(cbi->msgType)
+        {
+        case 1:
+            FIOSDeviceHelper::DoDeviceConnect(CallbackInfo);
+            break;
+
+        case 2:
+            FIOSDeviceHelper::DoDeviceDisconnect(CallbackInfo);
+            break;
+        }
+    }
+}
+
+void FIOSDeviceHelper::DoDeviceConnect(void* CallbackInfo)
+{
+    // connect to the device
+    struct FDeviceNotificationCallbackInformation* cbi = (FDeviceNotificationCallbackInformation*)CallbackInfo;
+    FIOSDevice* Device = new FIOSDevice(cbi->UDID, cbi->DeviceName);
+
+    // fire the event
+    FIOSLaunchDaemonPong Event;
+    Event.DeviceID = FString::Printf(TEXT("%s@%s"), cbi->ProductType.Contains(TEXT("AppleTV")) ? TEXT("TVOS") : TEXT("IOS"), *(cbi->UDID));
+    Event.DeviceName = cbi->DeviceName;
+    Event.DeviceType = cbi->ProductType;
+    Event.bCanReboot = false;
+    Event.bCanPowerOn = false;
+    Event.bCanPowerOff = false;
+    FIOSDeviceHelper::OnDeviceConnected().Broadcast(Event);
+
+    // add to the device list
+    ConnectedDevices.Add(Device, Event);
+}
+
+void FIOSDeviceHelper::DoDeviceDisconnect(void* CallbackInfo)
+{
+    struct FDeviceNotificationCallbackInformation* cbi = (FDeviceNotificationCallbackInformation*)CallbackInfo;
+    FIOSDevice* device = NULL;
+    for (auto DeviceIterator = ConnectedDevices.CreateIterator(); DeviceIterator; ++DeviceIterator)
+    {
+        if (DeviceIterator.Key()->SerialNumber() == cbi->UDID)
+        {
+            device = DeviceIterator.Key();
+            break;
+        }
+    }
+    if (device != NULL)
+    {
+        // extract the device id from the connected list
+        FIOSLaunchDaemonPong Event = ConnectedDevices.FindAndRemoveChecked(device);
+
+        // fire the event
+        FIOSDeviceHelper::OnDeviceDisconnected().Broadcast(Event);
+
+        // delete the device
+        delete device;
+    }
+}
+
+bool FIOSDeviceHelper::InstallIPAOnDevice(const FTargetDeviceId& DeviceId, const FString& IPAPath)
+{
+    return false;
+}
+
+void FIOSDeviceHelper::EnableDeviceCheck(bool OnOff)
+{
+    QueryTask->Enable(OnOff);
+}
diff --git a/Engine/Source/Developer/IOS/IOSTargetPlatform/Private/Mac/IOSDeviceHelperMac.cpp b/Engine/Source/Developer/IOS/IOSTargetPlatform/Private/Mac/IOSDeviceHelperMac.cpp
deleted file mode 100644
diff --git a/Engine/Source/Developer/IOS/IOSTargetPlatform/Private/Windows/IOSDeviceHelperWindows.cpp b/Engine/Source/Developer/IOS/IOSTargetPlatform/Private/Windows/IOSDeviceHelperWindows.cpp
deleted file mode 100644
