fixed compilation errors, e.g. https://horde.devtools.epicgames.com/job/6023cd915cb1b3000161e61c?label=4 that happens also for Win64

diff --git a/Engine/Source/Developer/IOS/IOSTargetPlatform/Private/IOSDeviceHelper.cpp b/Engine/Source/Developer/IOS/IOSTargetPlatform/Private/IOSDeviceHelper.cpp
--- a/Engine/Source/Developer/IOS/IOSTargetPlatform/Private/IOSDeviceHelper.cpp
+++ b/Engine/Source/Developer/IOS/IOSTargetPlatform/Private/IOSDeviceHelper.cpp
@@ -1,3 +1,4 @@
+#include "IOSDeviceHelper.h"
 #include "IOSTargetPlatform.h"
 #include "IOSTargetDeviceOutput.h"
 #include "HAL/PlatformProcess.h"
diff --git a/Engine/Source/Developer/IOS/IOSTargetPlatform/Private/IOSDeviceHelper.h b/Engine/Source/Developer/IOS/IOSTargetPlatform/Private/IOSDeviceHelper.h
--- a/Engine/Source/Developer/IOS/IOSTargetPlatform/Private/IOSDeviceHelper.h
+++ b/Engine/Source/Developer/IOS/IOSTargetPlatform/Private/IOSDeviceHelper.h
@@ -17,7 +17,7 @@ typedef TSharedRef<class FIOSDeviceHelper, ESPMode::ThreadSafe> FIOSDeviceHelper
  *
  * The first parameter is newly added or removed device
  */
-DECLARE_MULTICAST_DELEGATE_OneParam(FOnDeviceConnectEvent, const FIOSLaunchDaemonPong&)
+DECLARE_MULTICAST_DELEGATE_OneParam(FOnDeviceConnectEvent, const struct FIOSLaunchDaemonPong&);
 
 /**
  * Implemented the iOS device helper class
@@ -52,7 +52,7 @@ public:
     /**
      * Installs an ipa on to the device
      */
-    static bool InstallIPAOnDevice(const FTargetDeviceId& DeviceId, const FString& IPAPath);
+    static bool InstallIPAOnDevice(const class FTargetDeviceId& DeviceId, const FString& IPAPath);
     
 	/**
 	 * Suspends/Enables the device connect/disconnect thread
