Add option to override the path on the remote build Mac where iOS builds are performed.

diff --git a/Engine/Config/BaseEngine.ini b/Engine/Config/BaseEngine.ini
--- a/Engine/Config/BaseEngine.ini
+++ b/Engine/Config/BaseEngine.ini
@@ -2151,6 +2151,7 @@ bSupportsIPad=True
 bSupportsIPhone=True
 AdditionalPlistData=
 RemoteServerName=
+RemoteServerOverrideBuildPath=
 RSyncUsername=
 SSHPrivateKeyOverridePath=
 bEnableRemoteNotificationsSupport=False
diff --git a/Engine/Source/Developer/IOS/IOSPlatformEditor/Private/IOSTargetSettingsCustomization.cpp b/Engine/Source/Developer/IOS/IOSPlatformEditor/Private/IOSTargetSettingsCustomization.cpp
--- a/Engine/Source/Developer/IOS/IOSPlatformEditor/Private/IOSTargetSettingsCustomization.cpp
+++ b/Engine/Source/Developer/IOS/IOSPlatformEditor/Private/IOSTargetSettingsCustomization.cpp
@@ -903,7 +903,7 @@ void FIOSTargetSettingsCustomization::BuildRemoteBuildingSection(IDetailLayoutBu
 	TSharedRef<IPropertyHandle> RemoteServerNamePropertyHandle = DetailLayout.GetProperty(GET_MEMBER_NAME_CHECKED(UIOSRuntimeSettings, RemoteServerName));
 	IDetailPropertyRow& RemoteServerNamePropertyRow = RemoteBuildingGroup.AddPropertyRow(RemoteServerNamePropertyHandle);
 	RemoteServerNamePropertyRow
-		.ToolTip(LOCTEXT("RemoteServerNameToolTip", "The name or ip address of the remote mac which will be used to build IOS"))
+		.ToolTip(RemoteServerNamePropertyHandle->GetToolTipText())
 		.CustomWidget()
 		.NameContent()
 		[
@@ -948,7 +948,7 @@ void FIOSTargetSettingsCustomization::BuildRemoteBuildingSection(IDetailLayoutBu
 	TSharedRef<IPropertyHandle> RSyncUsernamePropertyHandle = DetailLayout.GetProperty(GET_MEMBER_NAME_CHECKED(UIOSRuntimeSettings, RSyncUsername));
 	IDetailPropertyRow& RSyncUsernamePropertyRow = RemoteBuildingGroup.AddPropertyRow(RSyncUsernamePropertyHandle);
 	RSyncUsernamePropertyRow
-		.ToolTip(LOCTEXT("RSyncUsernameToolTip", "The username of the mac user that matches the specified SSH Key."))
+		.ToolTip(RSyncUsernamePropertyHandle->GetToolTipText())
 		.CustomWidget()
 			.NameContent()
 			[
@@ -982,18 +982,17 @@ void FIOSTargetSettingsCustomization::BuildRemoteBuildingSection(IDetailLayoutBu
 			];
 
 
+	// Add Remote Server Override Build Path 
+	TSharedRef<IPropertyHandle> RemoteServerOverrideBuildPathPropertyHandle = DetailLayout.GetProperty(GET_MEMBER_NAME_CHECKED(UIOSRuntimeSettings, RemoteServerOverrideBuildPath));
+	IDetailPropertyRow& RemoteServerOverrideBuildPathPropertyRow = RemoteBuildingGroup.AddPropertyRow(RemoteServerOverrideBuildPathPropertyHandle);
+	RemoteServerOverrideBuildPathPropertyRow
+		.ToolTip(RemoteServerOverrideBuildPathPropertyHandle->GetToolTipText());
+
 	// Add existing SSH path label.
 	TSharedRef<IPropertyHandle> SSHPrivateKeyLocationPropertyHandle = DetailLayout.GetProperty(GET_MEMBER_NAME_CHECKED(UIOSRuntimeSettings, SSHPrivateKeyLocation));
 	IDetailPropertyRow& SSHPrivateKeyLocationPropertyRow = RemoteBuildingGroup.AddPropertyRow(SSHPrivateKeyLocationPropertyHandle);
 	SSHPrivateKeyLocationPropertyRow
-		.ToolTip(LOCTEXT("SSHPrivateKeyLocationToolTip", "The existing location of an SSH Key found by UE4."));
-
-
-	// Add SSH override path
-	TSharedRef<IPropertyHandle> SSHPrivateKeyOverridePathPropertyHandle = DetailLayout.GetProperty(GET_MEMBER_NAME_CHECKED(UIOSRuntimeSettings, SSHPrivateKeyOverridePath));
-	IDetailPropertyRow& SSHPrivateKeyOverridePathPropertyRow = RemoteBuildingGroup.AddPropertyRow(SSHPrivateKeyOverridePathPropertyHandle);
-	SSHPrivateKeyOverridePathPropertyRow
-		.ToolTip(LOCTEXT("SSHPrivateKeyOverridePathToolTip", "Override the existing SSH Private Key with one from a specified location."));
+		.ToolTip(SSHPrivateKeyLocationPropertyHandle->GetToolTipText());
 
 	// delta copy path
 	TSharedRef<IPropertyHandle> DeltaCopyOverridePathPropertyHandle = DetailLayout.GetProperty(GET_MEMBER_NAME_CHECKED(UIOSRuntimeSettings, DeltaCopyInstallPath));
diff --git a/Engine/Source/Programs/UnrealBuildTool/ToolChain/RemoteMac.cs b/Engine/Source/Programs/UnrealBuildTool/ToolChain/RemoteMac.cs
--- a/Engine/Source/Programs/UnrealBuildTool/ToolChain/RemoteMac.cs
+++ b/Engine/Source/Programs/UnrealBuildTool/ToolChain/RemoteMac.cs
@@ -271,12 +271,21 @@ namespace UnrealBuildTool
 			CommonRsyncArguments.Add("--prune-empty-dirs"); // Remove empty directories from the file list
 
 			// Get the remote base directory
-			StringBuilder Output;
-			if(ExecuteAndCaptureOutput("'echo ~'", out Output) != 0)
+			string RemoteServerOverrideBuildPath;
+			if (Ini.GetString("/Script/IOSRuntimeSettings.IOSRuntimeSettings", "RemoteServerOverrideBuildPath", out RemoteServerOverrideBuildPath) && !String.IsNullOrEmpty(RemoteServerOverrideBuildPath))
 			{
-				throw new BuildException("Unable to determine home directory for remote user. SSH output:\n{0}", StringUtils.Indent(Output.ToString(), "  "));
+				RemoteBaseDir = String.Format("{0}/{1}", RemoteServerOverrideBuildPath.Trim().TrimEnd('/'), Environment.MachineName);
 			}
-			RemoteBaseDir = String.Format("{0}/UE4/Builds/{1}", Output.ToString().Trim().TrimEnd('/'), Environment.MachineName);
+			else
+			{
+				StringBuilder Output;
+				if (ExecuteAndCaptureOutput("'echo ~'", out Output) != 0)
+				{
+					throw new BuildException("Unable to determine home directory for remote user. SSH output:\n{0}", StringUtils.Indent(Output.ToString(), "  "));
+				}
+				RemoteBaseDir = String.Format("{0}/UE5/Builds/{1}", Output.ToString().Trim().TrimEnd('/'), Environment.MachineName);
+			}
+
 			Log.TraceInformation("[Remote] Using base directory '{0}'", RemoteBaseDir);
 
 			// Build the list of directory mappings between the local and remote machines
diff --git a/Engine/Source/Runtime/IOS/IOSRuntimeSettings/Classes/IOSRuntimeSettings.h b/Engine/Source/Runtime/IOS/IOSRuntimeSettings/Classes/IOSRuntimeSettings.h
--- a/Engine/Source/Runtime/IOS/IOSRuntimeSettings/Classes/IOSRuntimeSettings.h
+++ b/Engine/Source/Runtime/IOS/IOSRuntimeSettings/Classes/IOSRuntimeSettings.h
@@ -270,6 +270,10 @@ public:
 	UPROPERTY(GlobalConfig, EditAnywhere, Category = "Build", meta = (EditCondition = "bUseRSync", DisplayName = "Username on Remote Server.", ConfigHierarchyEditable))
 	FString RSyncUsername;
 
+	// Optional path on the remote mac where the build files will be copied. If blank, ~/UE5/Builds will be used.
+	UPROPERTY(GlobalConfig, EditAnywhere, Category = "Build", meta = (ConfigHierarchyEditable))
+	FString RemoteServerOverrideBuildPath;
+
 	// The install directory of DeltaCopy.
 	UPROPERTY(GlobalConfig, EditAnywhere, Category = "Build", meta = (EditCondition = "bUseRSync", ConfigHierarchyEditable))
 	FIOSBuildResourceDirectory DeltaCopyInstallPath;
