Fix unity issue

#ue5
#ios

diff --git a/Engine/Source/Developer/IOS/IOSTargetPlatform/Private/IOSDeviceHelper.h b/Engine/Source/Developer/IOS/IOSTargetPlatform/Private/IOSDeviceHelper.h
--- a/Engine/Source/Developer/IOS/IOSTargetPlatform/Private/IOSDeviceHelper.h
+++ b/Engine/Source/Developer/IOS/IOSTargetPlatform/Private/IOSDeviceHelper.h
@@ -2,6 +2,8 @@
 
 #pragma once
 
+#include "CoreMinimal.h"
+
 /**
  * Type definition for shared pointers to instances of FIOSDeviceHelper
  */
