Disable compressed volume textures for TVOS, as min spec AppleTV HD with A8 GPU does not support them

diff --git a/Engine/Source/Developer/IOS/IOSTargetPlatform/Private/IOSTargetPlatform.cpp b/Engine/Source/Developer/IOS/IOSTargetPlatform/Private/IOSTargetPlatform.cpp
--- a/Engine/Source/Developer/IOS/IOSTargetPlatform/Private/IOSTargetPlatform.cpp
+++ b/Engine/Source/Developer/IOS/IOSTargetPlatform/Private/IOSTargetPlatform.cpp
@@ -586,7 +586,9 @@ void FIOSTargetPlatform::GetTextureFormats( const UTexture* Texture, TArray< TAr
 	if (TextureFormatNames.Num() == 0)
 	{
         int32 BlockSize = 1;
-		bool bSupportCompressedVolumeTexture = true;
+		// Compressed volume textures require MTLGPUFamilyApple3 or later
+		// min spec for TVOS is AppleTV HD which is MTLGPUFamilyApple2 (A8)
+		bool bSupportCompressedVolumeTexture = !bIsTVOS;
 		GetDefaultTextureFormatNamePerLayer(TextureFormatNames, this, Texture, true, bSupportCompressedVolumeTexture, BlockSize);
 	}
 
