Fix iOS device discovery in UnrealFrontend when run in standalone

diff --git a/Engine/Source/Developer/IOS/IOSTargetPlatform/Private/IOSDeviceHelper.cpp b/Engine/Source/Developer/IOS/IOSTargetPlatform/Private/IOSDeviceHelper.cpp
--- a/Engine/Source/Developer/IOS/IOSTargetPlatform/Private/IOSDeviceHelper.cpp
+++ b/Engine/Source/Developer/IOS/IOSTargetPlatform/Private/IOSDeviceHelper.cpp
@@ -219,6 +219,8 @@ public:
                         QueryDevices();
                     }
                 }
+#else
+				 QueryDevices();
 #endif
             }
             RecheckCounter--;
