Fix Min and Max Shaders version for iOS

diff --git a/Engine/Source/Developer/IOS/IOSPlatformEditor/Private/IOSTargetSettingsCustomization.cpp b/Engine/Source/Developer/IOS/IOSPlatformEditor/Private/IOSTargetSettingsCustomization.cpp
--- a/Engine/Source/Developer/IOS/IOSPlatformEditor/Private/IOSTargetSettingsCustomization.cpp
+++ b/Engine/Source/Developer/IOS/IOSPlatformEditor/Private/IOSTargetSettingsCustomization.cpp
@@ -64,27 +64,25 @@ FIOSTargetSettingsCustomization::FIOSTargetSettingsCustomization()
 	, EngineGraphicsPath(FString::Printf(TEXT("%sBuild/IOS/Resources/Graphics"), *FPaths::EngineDir()))
 	, GameGraphicsPath(FString::Printf(TEXT("%sBuild/IOS/Resources/Graphics"), *FPaths::ProjectDir()))
 {
-	new (IconNames) FPlatformIconInfo(TEXT("Icon20.png"), LOCTEXT("NotificationIcon_iPhone", "iPhone Notification Icon"), FText::GetEmpty(), 20, 20, FPlatformIconInfo::Optional);
-	new (IconNames) FPlatformIconInfo(TEXT("[REDACTED_EMAIL]"), LOCTEXT("NotificationIcon_iPhoneRetina", "iPhone Retina Notification Icon"), FText::GetEmpty(), 40, 40, FPlatformIconInfo::Optional);
-	new (IconNames) FPlatformIconInfo(TEXT("[REDACTED_EMAIL]"), LOCTEXT("NotificationIcon_iPhoneRetina_HD", "iPhone Retina HD Notification Icon"), FText::GetEmpty(), 60, 60, FPlatformIconInfo::Optional);
-	new (IconNames) FPlatformIconInfo(TEXT("Icon29.png"), LOCTEXT("SettingsIcon_iPhone", "iPhone Settings Icon"), FText::GetEmpty(), 29, 29, FPlatformIconInfo::Optional);// also iOS6 spotlight search
-	new (IconNames) FPlatformIconInfo(TEXT("[REDACTED_EMAIL]"), LOCTEXT("SettingsIcon_iPhoneRetina", "iPhone Retina Settings Icon"), FText::GetEmpty(), 58, 58, FPlatformIconInfo::Optional); // also iOS6 spotlight search
-	new (IconNames) FPlatformIconInfo(TEXT("[REDACTED_EMAIL]"), LOCTEXT("SettingsIcon_iPhoneRetina_HD", "iPhone Retina HD Settings Icon"), FText::GetEmpty(), 87, 87, FPlatformIconInfo::Optional); // also iOS6 spotlight search
-	new (IconNames) FPlatformIconInfo(TEXT("Icon40.png"), LOCTEXT("SpotlightIcon_iOS7", "iOS7 Spotlight Icon"), FText::GetEmpty(), 40, 40, FPlatformIconInfo::Optional);
-	new (IconNames) FPlatformIconInfo(TEXT("[REDACTED_EMAIL]"), LOCTEXT("SpotlightIcon_Retina_iOS7", "Retina iOS7 Spotlight Icon"), FText::GetEmpty(), 80, 80, FPlatformIconInfo::Optional);
-	new (IconNames) FPlatformIconInfo(TEXT("[REDACTED_EMAIL]"), LOCTEXT("SpotlightIcon_Retina_HD_iOS7", "Retina HD iOS7 Spotlight Icon"), FText::GetEmpty(), 120, 120, FPlatformIconInfo::Optional);
-	new (IconNames) FPlatformIconInfo(TEXT("Icon50.png"), LOCTEXT("SpotlightIcon_iPad_iOS6", "iPad iOS6 Spotlight Icon"), FText::GetEmpty(), 50, 50, FPlatformIconInfo::Optional);
-	new (IconNames) FPlatformIconInfo(TEXT("[REDACTED_EMAIL]"), LOCTEXT("SpotlightIcon_iPadRetina_iOS6", "iPad Retina iOS6 Spotlight Icon"), FText::GetEmpty(), 100, 100, FPlatformIconInfo::Optional);
-	new (IconNames) FPlatformIconInfo(TEXT("Icon57.png"), LOCTEXT("AppIcon_iPhone_iOS6", "iPhone iOS6 App Icon"), FText::GetEmpty(), 57, 57, FPlatformIconInfo::Required);
-	new (IconNames) FPlatformIconInfo(TEXT("[REDACTED_EMAIL]"), LOCTEXT("AppIcon_iPhoneRetina_iOS6", "iPhone Retina iOS6 App Icon"), FText::GetEmpty(), 114, 114, FPlatformIconInfo::Required);
-	new (IconNames) FPlatformIconInfo(TEXT("[REDACTED_EMAIL]"), LOCTEXT("AppIcon_iPhoneRetina_iOS7", "iPhone Retina iOS7 App Icon"), FText::GetEmpty(), 120, 120, FPlatformIconInfo::Required);
-	new (IconNames)FPlatformIconInfo(TEXT("[REDACTED_EMAIL]"), LOCTEXT("AppIcon_iPhoneRetina_iOS8", "iPhone Plus Retina iOS8 App Icon"), FText::GetEmpty(), 180, 180, FPlatformIconInfo::Required);
-	new (IconNames)FPlatformIconInfo(TEXT("Icon72.png"), LOCTEXT("AppIcon_iPad_iOS6", "iPad iOS6 App Icon"), FText::GetEmpty(), 72, 72, FPlatformIconInfo::Required);
-	new (IconNames) FPlatformIconInfo(TEXT("[REDACTED_EMAIL]"), LOCTEXT("AppIcon_iPadRetina_iOS6", "iPad Retina iOS6 App Icon"), FText::GetEmpty(), 144, 144, FPlatformIconInfo::Required);
-	new (IconNames) FPlatformIconInfo(TEXT("Icon76.png"), LOCTEXT("AppIcon_iPad_iOS7", "iPad iOS7 App Icon"), FText::GetEmpty(), 76, 76, FPlatformIconInfo::Required);
-	new (IconNames) FPlatformIconInfo(TEXT("[REDACTED_EMAIL]"), LOCTEXT("AppIcon_iPadRetina_iOS7", "iPad Retina iOS7 App Icon"), FText::GetEmpty(), 152, 152, FPlatformIconInfo::Required);
-	new (IconNames)FPlatformIconInfo(TEXT("[REDACTED_EMAIL]"), LOCTEXT("AppIcon_iPadProRetina_iOS9", "iPad Pro Retina iOS9 App Icon"), FText::GetEmpty(), 167, 167, FPlatformIconInfo::Required);
-	new (IconNames)FPlatformIconInfo(TEXT("Icon1024.png"), LOCTEXT("AppIcon_Marketing", "Marketing Icon"), FText::GetEmpty(), 1024, 1024, FPlatformIconInfo::Required);
+    // Default AppIcons copied at the payload's root. See https://developer.apple.com/design/human-interface-guidelines/ios/icons-and-images/app-icon/
+    new (IconNames) FPlatformIconInfo(TEXT("[REDACTED_EMAIL]"), LOCTEXT("Default_iPhone_AppIcon", "Default iPhone Icon"), FText::GetEmpty(), 120, 120, FPlatformIconInfo::Required); // iPhone
+    new (IconNames) FPlatformIconInfo(TEXT("[REDACTED_EMAIL]"), LOCTEXT("Default_iPad_AppIcon", "Default iPad AppIcon"), FText::GetEmpty(), 152, 152, FPlatformIconInfo::Required); // iPad, iPad Mini
+ 
+    //From here, all the icons are part of the asset catalog (Assets.car)
+    new (IconNames)FPlatformIconInfo(TEXT("Icon1024.png"), LOCTEXT("AppIcon_Marketing", "Marketing Icon"), FText::GetEmpty(), 1024, 1024, FPlatformIconInfo::Required); // App Store
+
+    new (IconNames)FPlatformIconInfo(TEXT("[REDACTED_EMAIL]"), LOCTEXT("3x_iPhone_App_Icon", "3x iPhone App Icon"), FText::GetEmpty(), 180, 180, FPlatformIconInfo::Optional); // iPhone
+    new (IconNames)FPlatformIconInfo(TEXT("[REDACTED_EMAIL]"), LOCTEXT("iPad_Pro_Retina_App_Icon", "iPad Pro Retina App Icon"), FText::GetEmpty(), 167, 167, FPlatformIconInfo::Optional); // iPad Pro
+
+    new (IconNames) FPlatformIconInfo(TEXT("[REDACTED_EMAIL]"), LOCTEXT("3x_iPhone_Spotlight_Icon", "3x iPhone Spotlight Icon"), FText::GetEmpty(), 120, 120, FPlatformIconInfo::Optional); // iPhone
+    new (IconNames) FPlatformIconInfo(TEXT("[REDACTED_EMAIL]"), LOCTEXT("Default_Spotlight_Icon", "Default Spotlight Icon"), FText::GetEmpty(), 80, 80, FPlatformIconInfo::Optional); // iPhone, iPad Pro, iPad, iPad Mini
+
+    new (IconNames) FPlatformIconInfo(TEXT("[REDACTED_EMAIL]"), LOCTEXT("3x_iPhone_Settings_Icon", "3x iPhone Settings Icon"), FText::GetEmpty(), 87, 87, FPlatformIconInfo::Optional); // iPhone
+    new (IconNames) FPlatformIconInfo(TEXT("[REDACTED_EMAIL]"), LOCTEXT("Default_Settings_Icon", "Default Settings Icon"), FText::GetEmpty(), 58, 58, FPlatformIconInfo::Optional); // iPhone, iPad Pro, iPad, iPad Mini
+
+    new (IconNames) FPlatformIconInfo(TEXT("[REDACTED_EMAIL]"), LOCTEXT("3x_iPhone_Notification_Icon", "3x iPhone Notification Icon"), FText::GetEmpty(), 60, 60, FPlatformIconInfo::Optional); // iPhone
+    new (IconNames) FPlatformIconInfo(TEXT("[REDACTED_EMAIL]"), LOCTEXT("Default_Notification_Icon", "Default Notification Icon"), FText::GetEmpty(), 40, 40, FPlatformIconInfo::Optional); // iPhone, iPad Pro, iPad, iPad Mini
+
 
 	new (LaunchImageNames)FPlatformIconInfo(TEXT("LaunchScreenIOS.png"), LOCTEXT("LaunchImageIOS", "Launch Screen Image"), LOCTEXT("LaunchImageIOSDesc", "This image is used for the Launch Screen when custom launch screen storyboards are not in use. The image is used in both portait and landscape modes and will be uniformly scaled to occupy the full width or height as necessary for of all devices, so if your app supports both a square image is recommended. The png file supplied must not have an alpha channel."), -1, -1, FPlatformIconInfo::Required);
 
@@ -1796,10 +1794,6 @@ void FIOSTargetSettingsCustomization::UpdateOSVersionWarning()
 			if (EnumValue < (uint8)EIOSVersion::IOS_13)
 			{
 				SetMinVersion((int32)EIOSVersion::IOS_13);
-				
-				FText Message;
-				Message = LOCTEXT("MetalMRTStandardv1.2","Enabling the Desktop Forward Renderer Metal requires Shader Standard v2.0 which increases the minimum operating system requirement for Metal from iOS 10.0 or later to iOS 11.0 or later.");
-				IOSVersionWarningTextBox->SetError(Message);
 			}
 		}
 		else
@@ -1824,20 +1818,12 @@ void FIOSTargetSettingsCustomization::UpdateMetalMRTWarning()
 			if (EnumValue < (uint8)EIOSVersion::IOS_13)
 			{
 				SetMinVersion((int32)EIOSVersion::IOS_13);
-				
-				FText Message;
-				Message = LOCTEXT("MetalMRTStandardv1.2","Enabling the Desktop Forward Renderer Metal requires Shader Standard v2.0 which increases the minimum operating system requirement for Metal from iOS 10.0 or later to iOS 11.0 or later.");
-				IOSVersionWarningTextBox->SetError(Message);
 			}
 			
 			ShaderVersionPropertyHandle->GetValue(EnumValue);
 			if (EnumValue < (uint8)EIOSMetalShaderStandard::IOSMetalSLStandard_2_2)
 			{
 				SetShaderStandard((int32)EIOSMetalShaderStandard::IOSMetalSLStandard_2_2);
-				
-				FText Message;
-				Message = LOCTEXT("MetalMRTStandardv1.2","Enabling the Desktop Forward Renderer Metal requires Shader Standard v2.0 which increases the minimum operating system requirement for Metal from iOS 10.0 or later to iOS 11.0 or later.");
-				ShaderVersionWarningTextBox->SetError(Message);
 			}
 		}
 		else
diff --git a/Engine/Source/Runtime/RHI/Private/RHI.cpp b/Engine/Source/Runtime/RHI/Private/RHI.cpp
--- a/Engine/Source/Runtime/RHI/Private/RHI.cpp
+++ b/Engine/Source/Runtime/RHI/Private/RHI.cpp
@@ -1621,8 +1621,8 @@ RHI_API uint32 RHIGetShaderLanguageVersion(const FStaticShaderPlatform Platform)
 			static int32 MaxShaderVersion = -1;
 			if (MaxShaderVersion < 0)
 			{
-				MaxShaderVersion = 5; // MSL v2.2
-				int32 MinShaderVersion = 4; // MSL v2.1
+				MaxShaderVersion = 6; // MSL v2.3
+				int32 MinShaderVersion = 5; // MSL v2.2
 				if(!GConfig->GetInt(TEXT("/Script/MacTargetPlatform.MacTargetSettings"), TEXT("MaxShaderLanguageVersion"), MaxShaderVersion, GEngineIni))
 				{
 					MaxShaderVersion = 0;
@@ -1636,8 +1636,8 @@ RHI_API uint32 RHIGetShaderLanguageVersion(const FStaticShaderPlatform Platform)
 			static int32 MaxShaderVersion = -1;
 			if (MaxShaderVersion < 0)
 			{
-				MaxShaderVersion = 2;
-				int32 MinShaderVersion = 2;
+				MaxShaderVersion = 6;
+				int32 MinShaderVersion = 5;
 				if(!GConfig->GetInt(TEXT("/Script/IOSRuntimeSettings.IOSRuntimeSettings"), TEXT("MaxShaderLanguageVersion"), MaxShaderVersion, GEngineIni))
 				{
 					MaxShaderVersion = 0;
