Remove Deployment Server references

diff --git a/Engine/Source/Developer/IOS/IOSTargetPlatform/Private/IOSTargetDevice.cpp b/Engine/Source/Developer/IOS/IOSTargetPlatform/Private/IOSTargetDevice.cpp
--- a/Engine/Source/Developer/IOS/IOSTargetPlatform/Private/IOSTargetDevice.cpp
+++ b/Engine/Source/Developer/IOS/IOSTargetPlatform/Private/IOSTargetDevice.cpp
@@ -14,273 +14,6 @@
 #include "Windows/AllowWindowsPlatformTypes.h"
 #endif // #if PLATFORM_WINDOWS
 
-enum
-{
-    DEFAULT_DS_COMMANDER_PORT = 41000, // default port to use when issuing DeploymentServer commands
-};
-
-FTcpDSCommander::FTcpDSCommander(const uint8* Data, int32 Count, TQueue<FString>& InOutputQueue)
-: bStopping(false)
-, bStopped(true)
-, bIsSuccess(false)
-, bIsSystemError(false)
-, DSSocket(nullptr)
-, Thread(nullptr)
-, OutputQueue(InOutputQueue)
-, DSCommand(nullptr)
-, DSCommandLen(Count + 1)
-, LastActivity(0.0)
-{
-    if (Count > 0)
-    {
-        DSCommand = (uint8*)FMemory::Malloc(Count + 1);
-        FPlatformMemory::Memcpy(DSCommand, Data, Count);
-        DSCommand[Count] = '\n';
-        Thread = FRunnableThread::Create(this, TEXT("FTcpDSCommander"), 128 * 1024, TPri_Normal);
-    }
-}
-
-/** Virtual destructor. */
-FTcpDSCommander::~FTcpDSCommander()
-{
-    if (Thread != nullptr)
-    {
-        Thread->Kill(true);
-        delete Thread;
-    }
-    if (DSCommand)
-    {
-        FMemory::Free(DSCommand);
-    }
-}
-
-bool FTcpDSCommander::Init()
-{
-	if (DSCommandLen < 1)
-	{
-		bIsSuccess = true;
-		return true;
-	}
-	ISocketSubsystem* SSS = ISocketSubsystem::Get();
-	DSSocket = SSS->CreateSocket(NAME_Stream, TEXT("DSCommander tcp"));
-	if (DSSocket == nullptr)
-	{
-		return false;
-	}
-	TSharedRef<FInternetAddr> Addr = SSS->CreateInternetAddr();
-	bool bIsValid;
-	Addr->SetIp(TEXT("127.0.0.1"), bIsValid);
-	Addr->SetPort(DEFAULT_DS_COMMANDER_PORT);
-
-#if PLATFORM_WINDOWS
-	// using the mutex to detect if the DeploymentServer is running
-	// only on windows
-	if (!IsDSRunning())
-	{
-		StartDSProcess();
-
-		int TimeoutSeconds = 5;
-		while (!IsDSRunning() && TimeoutSeconds > 0)
-		{
-			TimeoutSeconds--;
-			FPlatformProcess::Sleep(1.0f);
-		}
-		if (!IsDSRunning())
-		{
-			bIsSystemError = true;
-			return false;
-		}
-	}
-	if (DSSocket->Connect(*Addr) == false)
-	{
-		{ // extra bracked for cleaner ifdefs
-#else
-	// try to connect to the server
-	// on mac we use the old way to try a TCP connection
-	if (DSSocket->Connect(*Addr) == false)
-	{
-		StartDSProcess();
-		if (DSSocket->Connect(*Addr) == false)
-		{
-#endif // PLATFORM_WINDOWS
-        	// on failure, shut it all down
-        	SSS->DestroySocket(DSSocket);
-        	DSSocket = nullptr;
-			ESocketErrors LastError = SSS->GetLastErrorCode();
-			const TCHAR* SocketErr = SSS->GetSocketError(LastError);
-        	//UE_LOG(LogTemp, Display, TEXT("Failed to connect to deployment server at %s (%s)."), *Addr->ToString(true), SocketErr);
-        	return false;
-		}
-    }
-    LastActivity = FPlatformTime::Seconds();
-    
-    return true;
-}
-
-uint32 FTcpDSCommander::Run()
-{
-    if (!DSSocket)
-    {
-        //UE_LOG(LogTemp, Log, TEXT("Socket not created."));
-        return 1;
-    }
-    int32 NSent = 0;
-    bool BSent = DSSocket->Send(DSCommand, DSCommandLen, NSent);
-    if (NSent != DSCommandLen || !BSent)
-    {
-        //UE_LOG(LogTemp, Log, TEXT("Socket send error."));
-        return 1;
-    }
-    bStopped = false;
-    
-    static const SIZE_T BufferSize = 1025;
-	TArray<ANSICHAR> RecvBuffer;
-	RecvBuffer.Empty(BufferSize);
-	RecvBuffer.AddZeroed(BufferSize);
-
-    while (!bStopping)
-    {
-        uint32 Pending = 0;
-        if (DSSocket->GetConnectionState() != ESocketConnectionState::SCS_Connected)
-        {
-            //UE_LOG(LogTemp, Log, TEXT("Socket connection error."));
-            return 1;
-        }
-
-		int32 BytesRead = 0;
-		// Receive into the last BufferSize bytes of the buffer.
-		// We only receive at max BufferSize-1 bytes to preserve the last byte for a null terminator
-		if (DSSocket->Recv((uint8*)RecvBuffer.GetData() + RecvBuffer.Num() - BufferSize, BufferSize - 1, BytesRead) == false)
-		{
-			// Recv returns false on graceful socket disconnection
-			return 1;
-		}
-
-		if (BytesRead > 0)
-        {
-			LastActivity = FPlatformTime::Seconds();
-
-			ANSICHAR* LineStart = RecvBuffer.GetData();
-			ANSICHAR* LineEnd = nullptr;
-
-			while ((LineEnd = FCStringAnsi::Strchr(LineStart, '\n')) != nullptr)
-			{				
-				*LineEnd = '\0';
-				FString DataLine = UTF8_TO_TCHAR(LineStart);
-				LineStart = LineEnd + 1;
-
-				if (DataLine.EndsWith(TEXT("CMDOK\r")))
-				{
-					bIsSuccess = true;
-					return 0;
-				}
-				else if (DataLine.StartsWith(TEXT("[DSDIR]")))
-				{
-					// just ignore the folder check
-				}
-				else if (DataLine.EndsWith(TEXT("CMDFAIL\r")))
-				{
-					//UE_LOG(LogTemp, Display, TEXT("Socket command failed."));
-					return 1;
-				}
-				else
-				{
-					OutputQueue.Enqueue(DataLine);
-				}
-			}
-
-			if (*LineStart == '\0')
-			{
-				// no trailing data, reset buffer
-				RecvBuffer.Empty(BufferSize);
-				RecvBuffer.AddZeroed(BufferSize);
-			}
-			else
-			{
-				// There is trailing data, leave at the front of the buffer and add space for the next recv call.
-				RecvBuffer.RemoveAt(0, LineStart - RecvBuffer.GetData(), false);
-				// New buffer needs to have BufferSize free after the remaining string.
-				RecvBuffer.AddZeroed(BufferSize + FCStringAnsi::Strlen(RecvBuffer.GetData()) - RecvBuffer.Num());
-			}
-		}
-		else
-		{
-			double CurrentTime = FPlatformTime::Seconds();
-			if (CurrentTime - LastActivity > 120.0)
-			{
-				//UE_LOG(LogTemp, Display, TEXT("Socket command timeouted."));
-				return 0;
-			}
-			FPlatformProcess::Sleep(0.05f);
-		}
-    }
-    
-    return 0;
-}
-
-void FTcpDSCommander::Stop()
-{ 
-    bStopping = true;
-}
-
-void FTcpDSCommander::Exit()
-{
-	if (DSSocket)
-	{
-		DSSocket->Shutdown(ESocketShutdownMode::ReadWrite);
-		DSSocket->Close();
-		if (ISocketSubsystem::Get())
-		{
-			ISocketSubsystem::Get()->DestroySocket(DSSocket);
-		}
-	}
-	DSSocket = nullptr;
-
-	bStopped = true;
-}
-
-bool FTcpDSCommander::IsDSRunning()
-{
-	// is there a mutex we can use to connect test DS is running, also available on mac?
-	// there is also a failsafe mechanism for this, since the DeploymentServer will not start a new server if one is already running
-#if PLATFORM_WINDOWS
-	HANDLE mutex = CreateMutexA(NULL, TRUE, "Global\\DeploymentServer_Mutex_SERVERINSTANCE");
-	if (mutex == NULL || GetLastError() == ERROR_ALREADY_EXISTS)
-	{
-		// deployment server instance already runnning
-		if (mutex)
-		{
-			CloseHandle(mutex);
-		}
-		return true;
-	}
-	CloseHandle(mutex);
-#endif // PLATFORM_WINDOWS
-	return false;
-}
-
-bool FTcpDSCommander::StartDSProcess()
-{   
-    FString DSFilename = FPaths::ConvertRelativePathToFull(FPaths::EngineDir() / TEXT("Binaries/DotNET/IOS/DeploymentServerLauncher.exe"));
-    FString WorkingFoder = FPaths::ConvertRelativePathToFull(FPaths::EngineDir() / TEXT("Binaries/DotNET/IOS/"));
-    FString Params = "";
-    
-#if PLATFORM_MAC
-    // On Mac we launch UBT with Mono
-    FString ScriptPath = FPaths::ConvertRelativePathToFull(FPaths::EngineDir() / TEXT("Build/BatchFiles/Mac/RunMono.sh"));
-    Params = FString::Printf(TEXT("\"%s\" \"%s\" %s"), *ScriptPath, *DSFilename, *Params);
-    DSFilename = TEXT("/bin/sh");
-#endif
-    //UE_LOG(LogTemp, Log, TEXT("DeploymentServer not running, Starting it!"));
-    FPlatformProcess::CreateProc(*DSFilename, *Params, true, true, true, NULL, 0, *WorkingFoder, (void*)nullptr);
-    FPlatformProcess::Sleep(1.0f);
-    
-    return true;
-}
-
-
-//**********************************************************************************************************************************************************
-//*
 FIOSTargetDevice::FIOSTargetDevice(const ITargetPlatform& InTargetPlatform)
 	: TargetPlatform(InTargetPlatform)
 	, DeviceEndpoint()
@@ -406,16 +139,7 @@ bool FIOSTargetDevice::GetUserCredentials(FString& OutUserName, FString& OutUser
 
 inline void FIOSTargetDevice::ExecuteConsoleCommand(const FString& ExecCommand) const
 {
-	FString Params = FString::Printf(TEXT("command -device %s -param \"%s\""), *DeviceId.GetDeviceName(), *ExecCommand);
-
-	AsyncTask(
-		ENamedThreads::AnyThread,
-		[Params]()
-		{
-			FString StdOut;
-			FIOSTargetDeviceOutput::ExecuteDSCommand(TCHAR_TO_ANSI(*Params), &StdOut);
-		}
-	);
+	// this is where the new console command system execution goes.
 }
 
 inline ITargetDeviceOutputPtr FIOSTargetDevice::CreateDeviceOutputRouter(FOutputDevice* Output) const
@@ -427,44 +151,4 @@ inline ITargetDeviceOutputPtr FIOSTargetDevice::CreateDeviceOutputRouter(FOutput
 	}
 
 	return nullptr;
-}
-
-int FIOSTargetDeviceOutput::ExecuteDSCommand(const char *CommandLine, FString* OutStdOut)
-{
-	TQueue<FString>	OutputQueue;
-	FTcpDSCommander DSCommander((uint8*)CommandLine, strlen(CommandLine), OutputQueue);
-	while (DSCommander.IsValid() && !DSCommander.IsStopped())
-	{
-		FString NewLine;
-		if (OutputQueue.Dequeue(NewLine))
-		{
-			// process the string to break it up in to lines
-			*OutStdOut += NewLine;
-		}
-		else
-		{
-			FPlatformProcess::Sleep(0.05f);
-		}
-	}
-	FString NewLine;
-	if (OutputQueue.Dequeue(NewLine))
-	{
-		// process the string to break it up in to lines
-		*OutStdOut += NewLine;
-	}
-
-	FPlatformProcess::Sleep(0.05f);
-
-	if (DSCommander.IsSystemError())
-	{
-		return -1;
-	}
-
-	if (!DSCommander.WasSuccess())
-	{
-		//FPlatformMisc::LowLevelOutputDebugStringf(TEXT("The DeploymentServer command '%s' failed to run.\n"), ANSI_TO_TCHAR(CommandLine));
-		return 0;
-	}
-
-	return 1;
-}
+}
\ No newline at end of file
diff --git a/Engine/Source/Developer/IOS/IOSTargetPlatform/Private/IOSTargetDeviceOutput.h b/Engine/Source/Developer/IOS/IOSTargetPlatform/Private/IOSTargetDeviceOutput.h
--- a/Engine/Source/Developer/IOS/IOSTargetPlatform/Private/IOSTargetDeviceOutput.h
+++ b/Engine/Source/Developer/IOS/IOSTargetPlatform/Private/IOSTargetDeviceOutput.h
@@ -23,9 +23,6 @@ public:
 	virtual void Stop(void) override;
 	virtual uint32 Run(void) override;
 
-private:
-	bool StartDSCommander();
-
 private:
 	// > 0 if we've been asked to abort work in progress at the next opportunity
 	FThreadSafeCounter	StopTaskCounter;
@@ -33,7 +30,6 @@ private:
 	FTargetDeviceId		DeviceId;
 	FOutputDevice*		Output;
 	TQueue<FString>		OutputQueue;
-	FTcpDSCommander*	DSCommander;
 };
 
 /**
@@ -43,8 +39,6 @@ class FIOSTargetDeviceOutput : public ITargetDeviceOutput
 {
 public:
 	bool Init(const FIOSTargetDevice& TargetDevice, FOutputDevice* Output);
-
-	static int ExecuteDSCommand(const char *CommandLine, FString* OutStdOut);
 	
 private:
 	TUniquePtr<FRunnableThread>						DeviceOutputThread;
diff --git a/Engine/Source/Developer/IOS/IOSTargetPlatform/Private/IOSTargetDeviceOutput.inl b/Engine/Source/Developer/IOS/IOSTargetPlatform/Private/IOSTargetDeviceOutput.inl
--- a/Engine/Source/Developer/IOS/IOSTargetPlatform/Private/IOSTargetDeviceOutput.inl
+++ b/Engine/Source/Developer/IOS/IOSTargetPlatform/Private/IOSTargetDeviceOutput.inl
@@ -17,39 +17,18 @@ inline FIOSDeviceOutputReaderRunnable::FIOSDeviceOutputReaderRunnable(const FTar
 	: StopTaskCounter(0)
 	, DeviceId(InDeviceId)
 	, Output(InOutput)
-	, DSCommander(nullptr)
 {
 }
 
-inline bool FIOSDeviceOutputReaderRunnable::StartDSCommander()
-{
-	if (DSCommander)
-	{
-		DSCommander->Stop();
-		delete DSCommander;
-	}
-	Output->Serialize(TEXT("Starting listening ....."), ELogVerbosity::Log, NAME_None);
-	Output->Serialize(*DeviceId.GetDeviceName(), ELogVerbosity::Log, NAME_None);
-	FString Command = FString::Printf(TEXT("listentodevice -device %s"), *DeviceId.GetDeviceName());
-	uint8* DSCommand = (uint8*)TCHAR_TO_UTF8(*Command);
-	DSCommander = new FTcpDSCommander(DSCommand, strlen((const char*)DSCommand), OutputQueue);
-	return DSCommander->IsValid();
-}
 
 inline bool FIOSDeviceOutputReaderRunnable::Init(void) 
 { 
-	return StartDSCommander();
+	return true;
 }
 
 inline void FIOSDeviceOutputReaderRunnable::Exit(void) 
 {
 	StopTaskCounter.Increment();
-	if (DSCommander)
-	{
-		DSCommander->Stop();
-		delete DSCommander;
-		DSCommander = nullptr;
-	}
 }
 
 inline void FIOSDeviceOutputReaderRunnable::Stop(void)
@@ -60,36 +39,19 @@ inline void FIOSDeviceOutputReaderRunnable::Stop(void)
 inline uint32 FIOSDeviceOutputReaderRunnable::Run(void)
 {
 	Output->Serialize(TEXT("Starting Output"), ELogVerbosity::Log, NAME_None);
-	while (StopTaskCounter.GetValue() == 0 && DSCommander->IsValid())
+	while (StopTaskCounter.GetValue() == 0)
 	{
-		if (DSCommander->IsStopped() || !DSCommander->IsValid())
+		FString Text;
+		if (OutputQueue.Dequeue(Text))
 		{
-			// When user plugs out USB cable DS process stops
-			// Keep trying to restore DS connection until code that uses this object will not kill us
-			Output->Serialize(TEXT("Trying to restore connection to device..."), ELogVerbosity::Log, NAME_None);
-			if (StartDSCommander())
+			if (Text.Contains(TEXT("[UE]"), ESearchCase::CaseSensitive))
 			{
-				FPlatformProcess::Sleep(5.0f);
-			}
-			else
-			{
-				Output->Serialize(TEXT("Failed to start DS commander"), ELogVerbosity::Log, NAME_None);
+				Output->Serialize(*Text, ELogVerbosity::Log, NAME_None);
 			}
 		}
 		else
 		{
-			FString Text;
-			if (OutputQueue.Dequeue(Text))
-			{
-				if (Text.Contains(TEXT("[UE]"), ESearchCase::CaseSensitive))
-				{
-					Output->Serialize(*Text, ELogVerbosity::Log, NAME_None);
-				}
-			}
-			else
-			{
-				FPlatformProcess::Sleep(0.1f);
-			}
+			FPlatformProcess::Sleep(0.1f);
 		}
 	}
 
diff --git a/Engine/Source/Programs/IOS/iPhonePackager/DeployTime.cs b/Engine/Source/Programs/IOS/iPhonePackager/DeployTime.cs
--- a/Engine/Source/Programs/IOS/iPhonePackager/DeployTime.cs
+++ b/Engine/Source/Programs/IOS/iPhonePackager/DeployTime.cs
@@ -37,7 +37,6 @@ namespace iPhonePackager
 		}
 	}
 
-	// Implementation of DeploymentServer -> Application interface
 	[Serializable]
 	class DeployTimeReporter
 	{
