Refactor iOS Icons : create asset catalog according  to present file

diff --git a/Engine/Build/IOS/Resources/Assets.xcassets/AppIcon.appiconset/Contents.json b/Engine/Build/IOS/Resources/Assets.xcassets/AppIcon.appiconset/Contents.json
--- a/Engine/Build/IOS/Resources/Assets.xcassets/AppIcon.appiconset/Contents.json
+++ b/Engine/Build/IOS/Resources/Assets.xcassets/AppIcon.appiconset/Contents.json
@@ -1,71 +1,11 @@
 {
   "images" : [
-    {
-      "size" : "20x20",
-      "idiom" : "iphone",
-      "filename" : "[REDACTED_EMAIL]",
-      "scale" : "2x"
-    },
-    {
-      "size" : "20x20",
-      "idiom" : "iphone",
-      "filename" : "[REDACTED_EMAIL]",
-      "scale" : "3x"
-    },
-    {
-      "size" : "29x29",
-      "idiom" : "iphone",
-      "filename" : "[REDACTED_EMAIL]",
-      "scale" : "2x"
-    },
-    {
-      "size" : "29x29",
-      "idiom" : "iphone",
-      "filename" : "[REDACTED_EMAIL]",
-      "scale" : "3x"
-    },
-    {
-      "size" : "40x40",
-      "idiom" : "iphone",
-      "filename" : "[REDACTED_EMAIL]",
-      "scale" : "2x"
-    },
-    {
-      "size" : "40x40",
-      "idiom" : "iphone",
-      "filename" : "[REDACTED_EMAIL]",
-      "scale" : "3x"
-    },
     {
       "size" : "60x60",
       "idiom" : "iphone",
       "filename" : "[REDACTED_EMAIL]",
       "scale" : "2x"
     },
-    {
-      "size" : "60x60",
-      "idiom" : "iphone",
-      "filename" : "[REDACTED_EMAIL]",
-      "scale" : "3x"
-    },
-    {
-      "size" : "20x20",
-      "idiom" : "ipad",
-      "filename" : "[REDACTED_EMAIL]",
-      "scale" : "2x"
-    },
-    {
-      "size" : "29x29",
-      "idiom" : "ipad",
-      "filename" : "[REDACTED_EMAIL]",
-      "scale" : "2x"
-    },
-    {
-      "size" : "40x40",
-      "idiom" : "ipad",
-      "filename" : "[REDACTED_EMAIL]",
-      "scale" : "2x"
-    },
     {
       "size" : "76x76",
       "idiom" : "ipad",
diff --git a/Engine/Source/Developer/ExternalImagePicker/Private/SExternalImagePicker.cpp b/Engine/Source/Developer/ExternalImagePicker/Private/SExternalImagePicker.cpp
--- a/Engine/Source/Developer/ExternalImagePicker/Private/SExternalImagePicker.cpp
+++ b/Engine/Source/Developer/ExternalImagePicker/Private/SExternalImagePicker.cpp
@@ -206,6 +206,13 @@ void SExternalImagePicker::ApplyImage()
 	}
 	else
 	{
+        // This is preventing the "no default image" error display when we don't want a fallback behavior, i.e. iOS optional icons
+        FString DefaultImagePathString(DefaultImagePath);
+        if (DefaultImagePathString.IsEmpty())
+        {
+            return;
+        }
+            
 		TypeOfImage = UsingDummyPlaceholderImage;
 
 		if (bRequiresSpecificSize)
diff --git a/Engine/Source/Developer/IOS/IOSPlatformEditor/Private/IOSTargetSettingsCustomization.cpp b/Engine/Source/Developer/IOS/IOSPlatformEditor/Private/IOSTargetSettingsCustomization.cpp
--- a/Engine/Source/Developer/IOS/IOSPlatformEditor/Private/IOSTargetSettingsCustomization.cpp
+++ b/Engine/Source/Developer/IOS/IOSPlatformEditor/Private/IOSTargetSettingsCustomization.cpp
@@ -66,13 +66,16 @@ FIOSTargetSettingsCustomization::FIOSTargetSettingsCustomization()
 {
     // Default AppIcons copied at the payload's root. See https://developer.apple.com/design/human-interface-guidelines/ios/icons-and-images/app-icon/
     new (IconNames) FPlatformIconInfo(TEXT("[REDACTED_EMAIL]"), LOCTEXT("Default_iPhone_AppIcon", "Default iPhone Icon"), FText::GetEmpty(), 120, 120, FPlatformIconInfo::Required); // iPhone
-    new (IconNames) FPlatformIconInfo(TEXT("[REDACTED_EMAIL]"), LOCTEXT("Default_iPad_AppIcon", "Default iPad AppIcon"), FText::GetEmpty(), 152, 152, FPlatformIconInfo::Required); // iPad, iPad Mini
+    new (IconNames) FPlatformIconInfo(TEXT("[REDACTED_EMAIL]"), LOCTEXT("Default_iPad_AppIcon", "Default iPad App Icon"), FText::GetEmpty(), 152, 152, FPlatformIconInfo::Required); // iPad, iPad Mini
  
     //From here, all the icons are part of the asset catalog (Assets.car)
+
+    // Required in the asset catalog
     new (IconNames)FPlatformIconInfo(TEXT("Icon1024.png"), LOCTEXT("AppIcon_Marketing", "Marketing Icon"), FText::GetEmpty(), 1024, 1024, FPlatformIconInfo::Required); // App Store
+    new (IconNames)FPlatformIconInfo(TEXT("[REDACTED_EMAIL]"), LOCTEXT("iPad_Pro_Retina_App_Icon", "iPad Pro Retina App Icon"), FText::GetEmpty(), 167, 167, FPlatformIconInfo::Required); // iPad Pro
 
+    // From here icons are optional
     new (IconNames)FPlatformIconInfo(TEXT("[REDACTED_EMAIL]"), LOCTEXT("3x_iPhone_App_Icon", "3x iPhone App Icon"), FText::GetEmpty(), 180, 180, FPlatformIconInfo::Optional); // iPhone
-    new (IconNames)FPlatformIconInfo(TEXT("[REDACTED_EMAIL]"), LOCTEXT("iPad_Pro_Retina_App_Icon", "iPad Pro Retina App Icon"), FText::GetEmpty(), 167, 167, FPlatformIconInfo::Optional); // iPad Pro
 
     new (IconNames) FPlatformIconInfo(TEXT("[REDACTED_EMAIL]"), LOCTEXT("3x_iPhone_Spotlight_Icon", "3x iPhone Spotlight Icon"), FText::GetEmpty(), 120, 120, FPlatformIconInfo::Optional); // iPhone
     new (IconNames) FPlatformIconInfo(TEXT("[REDACTED_EMAIL]"), LOCTEXT("Default_Spotlight_Icon", "Default Spotlight Icon"), FText::GetEmpty(), 80, 80, FPlatformIconInfo::Optional); // iPhone, iPad Pro, iPad, iPad Mini
@@ -83,7 +86,6 @@ FIOSTargetSettingsCustomization::FIOSTargetSettingsCustomization()
     new (IconNames) FPlatformIconInfo(TEXT("[REDACTED_EMAIL]"), LOCTEXT("3x_iPhone_Notification_Icon", "3x iPhone Notification Icon"), FText::GetEmpty(), 60, 60, FPlatformIconInfo::Optional); // iPhone
     new (IconNames) FPlatformIconInfo(TEXT("[REDACTED_EMAIL]"), LOCTEXT("Default_Notification_Icon", "Default Notification Icon"), FText::GetEmpty(), 40, 40, FPlatformIconInfo::Optional); // iPhone, iPad Pro, iPad, iPad Mini
 
-
 	new (LaunchImageNames)FPlatformIconInfo(TEXT("LaunchScreenIOS.png"), LOCTEXT("LaunchImageIOS", "Launch Screen Image"), LOCTEXT("LaunchImageIOSDesc", "This image is used for the Launch Screen when custom launch screen storyboards are not in use. The image is used in both portait and landscape modes and will be uniformly scaled to occupy the full width or height as necessary for of all devices, so if your app supports both a square image is recommended. The png file supplied must not have an alpha channel."), -1, -1, FPlatformIconInfo::Required);
 
 	bShowAllProvisions = false;
@@ -1102,36 +1104,71 @@ void FIOSTargetSettingsCustomization::BuildImageRow(IDetailLayoutBuilder& Detail
 	const FString AutomaticImagePath = EngineGraphicsPath / Info.IconPath;
 	const FString TargetImagePath = GameGraphicsPath / Info.IconPath;
 
-	Category.AddCustomRow(Info.IconName)
-		.NameContent()
-		[
-			SNew(SHorizontalBox)
-			+ SHorizontalBox::Slot()
-			.Padding(FMargin(0, 1, 0, 1))
-			.FillWidth(1.0f)
-			[
-				SNew(STextBlock)
-				.Text(Info.IconName)
-				.Font(DetailLayout.GetDetailFont())
-				// IconDescription is not used, repurpose for tooltip
-				.ToolTipText(Info.IconDescription)
-			]
-		]
-		.ValueContent()
-		.MaxDesiredWidth(400.0f)
-		.MinDesiredWidth(100.0f)
-		[
-			SNew(SHorizontalBox)
-			+ SHorizontalBox::Slot()
-			.FillWidth(1.0f)
-			.VAlign(VAlign_Center)
-			[
-				SNew(SExternalImageReference, AutomaticImagePath, TargetImagePath)
-				.RequiredSize(Info.IconRequiredSize)
-				.MaxDisplaySize(MaxDisplaySize)
-				.DeleteTargetWhenDefaultChosen(true)
-			]
-		];
+    if (Info.RequiredState == FPlatformIconInfo::Required)
+    {
+        Category.AddCustomRow(Info.IconName)
+            .NameContent()
+            [
+                SNew(SHorizontalBox)
+                + SHorizontalBox::Slot()
+                .Padding(FMargin(0, 1, 0, 1))
+                .FillWidth(1.0f)
+                [
+                    SNew(STextBlock)
+                    .Text(Info.IconName)
+                    .Font(DetailLayout.GetDetailFont())
+                    // IconDescription is not used, repurpose for tooltip
+                    .ToolTipText(Info.IconDescription)
+                ]
+            ]
+            .ValueContent()
+            .MaxDesiredWidth(400.0f)
+            .MinDesiredWidth(100.0f)
+            [
+                SNew(SHorizontalBox)
+                + SHorizontalBox::Slot()
+                .FillWidth(1.0f)
+                .VAlign(VAlign_Center)
+                [
+                    SNew(SExternalImageReference, AutomaticImagePath, TargetImagePath)
+                    .RequiredSize(Info.IconRequiredSize)
+                    .MaxDisplaySize(MaxDisplaySize)
+                    .DeleteTargetWhenDefaultChosen(true)
+                ]
+            ];
+    }
+    else
+    {
+        Category.AddCustomRow(Info.IconName)
+            .NameContent()
+            [
+                SNew(SHorizontalBox)
+                + SHorizontalBox::Slot()
+                .Padding(FMargin(0, 1, 0, 1))
+                .FillWidth(1.0f)
+                [
+                    SNew(STextBlock)
+                    .Text(Info.IconName)
+                    .Font(DetailLayout.GetDetailFont())
+                    // IconDescription is not used, repurpose for tooltip
+                    .ToolTipText(Info.IconDescription)
+                ]
+            ]
+            .ValueContent()
+            .MaxDesiredWidth(400.0f)
+            .MinDesiredWidth(100.0f)
+            [
+                SNew(SHorizontalBox)
+                + SHorizontalBox::Slot()
+                .FillWidth(1.0f)
+                .VAlign(VAlign_Center)
+                [
+                    SNew(SExternalImageReference, "", TargetImagePath)
+                    .RequiredSize(Info.IconRequiredSize)
+                    .MaxDisplaySize(MaxDisplaySize)
+                ]
+            ];
+    }
 }
 
 void FIOSTargetSettingsCustomization::FindRequiredFiles()
diff --git a/Engine/Source/Programs/UnrealBuildTool/Platform/IOS/IOSToolChain.cs b/Engine/Source/Programs/UnrealBuildTool/Platform/IOS/IOSToolChain.cs
--- a/Engine/Source/Programs/UnrealBuildTool/Platform/IOS/IOSToolChain.cs
+++ b/Engine/Source/Programs/UnrealBuildTool/Platform/IOS/IOSToolChain.cs
@@ -1401,6 +1401,72 @@ namespace UnrealBuildTool
 						DestFileInfo.Attributes = DestFileInfo.Attributes & ~FileAttributes.ReadOnly;
 					}
 				}
+
+				StringBuilder ContentsJson = new StringBuilder();
+
+				ContentsJson.AppendLine("{");
+				ContentsJson.AppendLine("\"images\" : [");
+				ContentsJson.AppendLine("{");
+				ContentsJson.AppendLine("\"size\" : \"60x60\",");
+				ContentsJson.AppendLine("\"idiom\" : \"iphone\",");
+				ContentsJson.AppendLine("\"filename\" : \"[REDACTED_EMAIL]\",");
+				ContentsJson.AppendLine("\"scale\" : \"2x\"");
+				ContentsJson.AppendLine("},");
+				ContentsJson.AppendLine("{");
+				ContentsJson.AppendLine("\"size\" : \"76x76\",");
+				ContentsJson.AppendLine("\"idiom\" : \"ipad\",");
+				ContentsJson.AppendLine("\"filename\" : \"[REDACTED_EMAIL]\",");
+				ContentsJson.AppendLine("\"scale\" : \"2x\"");
+				ContentsJson.AppendLine("},");
+				ContentsJson.AppendLine("{");
+				ContentsJson.AppendLine("\"size\" : \"83.5x83.5\",");
+				ContentsJson.AppendLine("\"idiom\" : \"ipad\",");
+				ContentsJson.AppendLine("\"filename\" : \"[REDACTED_EMAIL]\",");
+				ContentsJson.AppendLine("\"scale\" : \"2x\"");
+				ContentsJson.AppendLine("},");
+				ContentsJson.AppendLine("{");
+				ContentsJson.AppendLine("\"size\" : \"1024x1024\",");
+				ContentsJson.AppendLine("\"idiom\" : \"ios-marketing\",");
+				ContentsJson.AppendLine("\"filename\" : \"Icon1024.png\",");
+				ContentsJson.AppendLine("\"scale\" : \"1x\"");
+				ContentsJson.AppendLine("},");
+
+				string[][] IconsInfo = {
+					new string []{ "[REDACTED_EMAIL]", "\"20x20\"",  "\"iphone\"", "\"2x\"" },
+					new string []{ "[REDACTED_EMAIL]", "\"20x20\"", "\"iphone\"", "\"3x\"" },
+					new string []{ "[REDACTED_EMAIL]", "\"29x29\"", "\"iphone\"", "\"2x\"" },
+					new string []{ "[REDACTED_EMAIL]", "\"29x29\"", "\"iphone\"", "\"3x\"" },
+					new string []{ "[REDACTED_EMAIL]", "\"40x40\"", "\"iphone\"", "\"2x\"" },
+					new string []{ "[REDACTED_EMAIL]", "\"40x40\"", "\"iphone\"", "\"3x\"" },
+					new string []{ "[REDACTED_EMAIL]", "\"60x60\"", "\"iphone\"", "\"3x\"" },
+					new string []{ "[REDACTED_EMAIL]", "\"20x20\"", "\"ipad\"", "\"2x\"" },
+					new string []{ "[REDACTED_EMAIL]", "\"29x29\"", "\"ipad\"", "\"2x\"" },
+					new string []{ "[REDACTED_EMAIL]", "\"40x40\"", "\"ipad\"", "\"2x\"" },
+					};
+
+				for (int Index = 0; Index < IconsInfo.Length; ++Index)
+				{
+					if (File.Exists(Path.Combine(Dir, IconsInfo[Index][0])))
+					{
+						ContentsJson.AppendLine("{");
+						ContentsJson.AppendLine("\"size\" : " + IconsInfo[Index][1] + ",");
+						ContentsJson.AppendLine("\"idiom\" : " + IconsInfo[Index][2] + ",");
+						ContentsJson.AppendLine("\"filename\" : \"" + IconsInfo[Index][0] + "\",");
+						ContentsJson.AppendLine("\"scale\" : " + IconsInfo[Index][3]);
+						ContentsJson.AppendLine("},");
+					}
+				}
+
+				ContentsJson.AppendLine("],");
+				ContentsJson.AppendLine("\"info\" : {");
+				ContentsJson.AppendLine("\"version\" : 1,");
+				ContentsJson.AppendLine("\"author\" : \"xcode\"");
+				ContentsJson.AppendLine("}");
+				ContentsJson.AppendLine("}");
+
+				string ContentsFile = Path.Combine(IntermediateDir, "Resources", "Assets.xcassets", "AppIcon.appiconset", "Contents.json");
+				File.WriteAllText(ContentsFile, ContentsJson.ToString());
+
 			}
 			return new DirectoryReference(ResourcesDir);
 		}
@@ -1449,7 +1515,7 @@ namespace UnrealBuildTool
 				OutputFiles.Add(StripCompleteFile);
 			}
 
-			if (!BinaryLinkEnvironment.bIsBuildingDLL && ShouldCompileAssetCatalog())
+			if (!BinaryLinkEnvironment.bIsBuildingDLL)
 			{
 				// generate the asset catalog
 				bool bUserImagesExist = false;
@@ -1535,11 +1601,6 @@ namespace UnrealBuildTool
 			return OutputFiles;
 		}
 
-		public bool ShouldCompileAssetCatalog()
-		{
-			return Target.Platform == UnrealTargetPlatform.TVOS || (Target.Platform == UnrealTargetPlatform.IOS && Settings.Value.IOSSDKVersionFloat >= 11.0f);
-		}
-
 		public static string GetCodesignPlatformName(UnrealTargetPlatform Platform)
 		{
 			if (Platform == UnrealTargetPlatform.TVOS)
