tvOS asset catalogue support in editor

diff --git a/Engine/Source/Developer/IOS/IOSPlatformEditor/Private/IOSTargetSettingsCustomization.cpp b/Engine/Source/Developer/IOS/IOSPlatformEditor/Private/IOSTargetSettingsCustomization.cpp
--- a/Engine/Source/Developer/IOS/IOSPlatformEditor/Private/IOSTargetSettingsCustomization.cpp
+++ b/Engine/Source/Developer/IOS/IOSPlatformEditor/Private/IOSTargetSettingsCustomization.cpp
@@ -63,6 +63,9 @@ FIOSTargetSettingsCustomization::FIOSTargetSettingsCustomization()
 	, GameInfoPath(FString::Printf(TEXT("%sBuild/IOS/Info.plist"), *FPaths::ProjectDir()))
 	, EngineGraphicsPath(FString::Printf(TEXT("%sBuild/IOS/Resources/Graphics"), *FPaths::EngineDir()))
 	, GameGraphicsPath(FString::Printf(TEXT("%sBuild/IOS/Resources/Graphics"), *FPaths::ProjectDir()))
+	, TVOSEngineGraphicsPath(FString::Printf(TEXT("%sBuild/TVOS/Resources/Graphics"), *FPaths::EngineDir()))
+	, TVOSGameGraphicsPath(FString::Printf(TEXT("%sBuild/TVOS/Resources/Graphics"), *FPaths::ProjectDir()))
+
 {
     // Default AppIcons copied at the payload's root. See https://developer.apple.com/design/human-interface-guidelines/ios/icons-and-images/app-icon/
     new (IconNames) FPlatformIconInfo(TEXT("[REDACTED_EMAIL]"), LOCTEXT("Default_iPhone_AppIcon", "Default iPhone Icon"), FText::GetEmpty(), 120, 120, FPlatformIconInfo::Required); // iPhone
@@ -86,8 +89,21 @@ FIOSTargetSettingsCustomization::FIOSTargetSettingsCustomization()
     new (IconNames) FPlatformIconInfo(TEXT("[REDACTED_EMAIL]"), LOCTEXT("3x_iPhone_Notification_Icon", "3x iPhone Notification Icon"), FText::GetEmpty(), 60, 60, FPlatformIconInfo::Optional); // iPhone
     new (IconNames) FPlatformIconInfo(TEXT("[REDACTED_EMAIL]"), LOCTEXT("Default_Notification_Icon", "Default Notification Icon"), FText::GetEmpty(), 40, 40, FPlatformIconInfo::Optional); // iPhone, iPad Pro, iPad, iPad Mini
 
+	// LaunchScreen iOS and tvOS
 	new (LaunchImageNames)FPlatformIconInfo(TEXT("LaunchScreenIOS.png"), LOCTEXT("LaunchImageIOS", "Launch Screen Image"), LOCTEXT("LaunchImageIOSDesc", "This image is used for the Launch Screen when custom launch screen storyboards are not in use. The image is used in both portait and landscape modes and will be uniformly scaled to occupy the full width or height as necessary for of all devices, so if your app supports both a square image is recommended. The png file supplied must not have an alpha channel."), -1, -1, FPlatformIconInfo::Required);
 
+	// Icons and Shelf Images for tvOS
+	new (TvOSImageNames)FPlatformIconInfo(TEXT("Icon_Large_Back.png"), LOCTEXT("TVOS_Icon_Large_Back", "Icon Large Back"), FText::GetEmpty(), 1280, 768, FPlatformIconInfo::Required);
+	new (TvOSImageNames)FPlatformIconInfo(TEXT("Icon_Large_Front.png"), LOCTEXT("TVOS_Icon_Large_Front", "Icon Large Front"), FText::GetEmpty(), 1280, 768, FPlatformIconInfo::Required);
+	new (TvOSImageNames)FPlatformIconInfo(TEXT("Icon_Large_Middle.png"), LOCTEXT("TVOS_Icon_Large_Middle", "Icon Large Middle"), FText::GetEmpty(), 1280, 768, FPlatformIconInfo::Required);
+	new (TvOSImageNames)FPlatformIconInfo(TEXT("Icon_Small_Back.png"), LOCTEXT("TVOS_Icon_Small_Back", "Icon Small Back"), FText::GetEmpty(), 400, 240, FPlatformIconInfo::Required);
+	new (TvOSImageNames)FPlatformIconInfo(TEXT("Icon_Small_Front.png"), LOCTEXT("TVOS_Icon_Small_Front", "Icon Small Front"), FText::GetEmpty(), 400, 240, FPlatformIconInfo::Required);
+	new (TvOSImageNames)FPlatformIconInfo(TEXT("Icon_Small_Middle.png"), LOCTEXT("TVOS_Icon_Small_Middle", "Icon Small Middle"), FText::GetEmpty(), 400, 240, FPlatformIconInfo::Required);
+	new (TvOSImageNames)FPlatformIconInfo(TEXT("TopShelf.png"), LOCTEXT("TVOS_Top_Shelf", "Top Shelf"), FText::GetEmpty(), 1920, 720, FPlatformIconInfo::Required);
+	new (TvOSImageNames)FPlatformIconInfo(TEXT("[REDACTED_EMAIL]"), LOCTEXT("2x_TVOS_Top_Shelf", "Top Shelf"), FText::GetEmpty(), 3840, 1440, FPlatformIconInfo::Required);
+	new (TvOSImageNames)FPlatformIconInfo(TEXT("TopShelfWide-1920x720.png"), LOCTEXT("TVOS_Top_Shelf_Wide", "Top Shelf Wide"), FText::GetEmpty(), 2320, 720, FPlatformIconInfo::Required);
+	new (TvOSImageNames)FPlatformIconInfo(TEXT("[REDACTED_EMAIL]"), LOCTEXT("2x_TVOS_Top_Shelf_Wide", "2x Top Shelf Wide"), FText::GetEmpty(), 4640, 1440, FPlatformIconInfo::Required);
+
 	bShowAllProvisions = false;
 	bShowAllCertificates = false;
 	ProvisionList = MakeShareable(new TArray<ProvisionPtr>());
@@ -1055,6 +1071,15 @@ void FIOSTargetSettingsCustomization::BuildIconSection(IDetailLayoutBuilder& Det
 	{
 		BuildImageRow(DetailLayout, LaunchImageCategory, Info, LaunchImageMaxSize);
 	}
+
+	// Add the tvOS content
+	IDetailCategoryBuilder& tvOSCategory = DetailLayout.EditCategory(FName("tvOS"));
+	const FVector2D TvOSImageMaxSize(150.0f, 150.0f);
+	for (const FPlatformIconInfo& Info : TvOSImageNames)
+	{
+		BuildImageRow(DetailLayout, tvOSCategory, Info, TvOSImageMaxSize, true);
+	}
+
 }
 
 
@@ -1099,10 +1124,16 @@ void FIOSTargetSettingsCustomization::CopySetupFilesIntoProject()
 	SavedLayoutBuilder->ForceRefreshDetails();
 }
 
-void FIOSTargetSettingsCustomization::BuildImageRow(IDetailLayoutBuilder& DetailLayout, IDetailCategoryBuilder& Category, const FPlatformIconInfo& Info, const FVector2D& MaxDisplaySize)
+void FIOSTargetSettingsCustomization::BuildImageRow(IDetailLayoutBuilder& DetailLayout, IDetailCategoryBuilder& Category, const FPlatformIconInfo& Info, const FVector2D& MaxDisplaySize, bool bIsTVOS)
 {
-	const FString AutomaticImagePath = EngineGraphicsPath / Info.IconPath;
-	const FString TargetImagePath = GameGraphicsPath / Info.IconPath;
+	FString AutomaticImagePath = EngineGraphicsPath / Info.IconPath;
+	FString TargetImagePath = GameGraphicsPath / Info.IconPath;
+
+	if (bIsTVOS)
+	{
+		AutomaticImagePath = TVOSEngineGraphicsPath / Info.IconPath;
+		TargetImagePath = TVOSGameGraphicsPath / Info.IconPath;
+	}
 
     if (Info.RequiredState == FPlatformIconInfo::Required)
     {
diff --git a/Engine/Source/Developer/IOS/IOSPlatformEditor/Private/IOSTargetSettingsCustomization.h b/Engine/Source/Developer/IOS/IOSPlatformEditor/Private/IOSTargetSettingsCustomization.h
--- a/Engine/Source/Developer/IOS/IOSPlatformEditor/Private/IOSTargetSettingsCustomization.h
+++ b/Engine/Source/Developer/IOS/IOSPlatformEditor/Private/IOSTargetSettingsCustomization.h
@@ -66,6 +66,7 @@ public:
 private:
 	TArray<struct FPlatformIconInfo> IconNames;
 	TArray<struct FPlatformIconInfo> LaunchImageNames;
+	TArray<struct FPlatformIconInfo> TvOSImageNames;
 
 	const FString EngineInfoPath;
 	const FString GameInfoPath;
@@ -73,6 +74,9 @@ private:
 	const FString EngineGraphicsPath;
 	const FString GameGraphicsPath;
 
+	const FString TVOSEngineGraphicsPath;
+	const FString TVOSGameGraphicsPath;
+
 	IDetailLayoutBuilder* SavedLayoutBuilder;
 
 	// Is the manifest writable?
@@ -120,7 +124,7 @@ private:
 	void CopySetupFilesIntoProject();
 
 	// Builds an image row
-	void BuildImageRow(class IDetailLayoutBuilder& DetailLayout, class IDetailCategoryBuilder& Category, const struct FPlatformIconInfo& Info, const FVector2D& MaxDisplaySize);
+	void BuildImageRow(class IDetailLayoutBuilder& DetailLayout, class IDetailCategoryBuilder& Category, const struct FPlatformIconInfo& Info, const FVector2D& MaxDisplaySize, bool bIsTVOS = false);
 
 	// Install the provision
  	FReply OnInstallProvisionClicked();
