Display the trust popup on iOS and tvOS when using an untrusted device

diff --git a/Engine/Source/Developer/IOS/IOSTargetPlatform/Private/IOSDeviceHelper.cpp b/Engine/Source/Developer/IOS/IOSTargetPlatform/Private/IOSDeviceHelper.cpp
--- a/Engine/Source/Developer/IOS/IOSTargetPlatform/Private/IOSDeviceHelper.cpp
+++ b/Engine/Source/Developer/IOS/IOSTargetPlatform/Private/IOSDeviceHelper.cpp
@@ -9,6 +9,9 @@
 #include "Interfaces/ITargetPlatformManagerModule.h"
 #include "Interfaces/ITargetPlatform.h"
 #include "Interfaces/IProjectManager.h"
+#include "Misc/MessageDialog.h"
+
+#define LOCTEXT_NAMESPACE "MiscIOSMessages"
 
 DEFINE_LOG_CATEGORY_STATIC(LogIOSDeviceHelper, Log, All);
 
@@ -27,6 +30,7 @@ DEFINE_LOG_CATEGORY_STATIC(LogIOSDeviceHelper, Log, All);
         FString ProductType;
         DeviceConnectionInterface DeviceInterface;
         uint32 msgType;
+		bool isAuthorized;
     };
 
 
@@ -35,6 +39,7 @@ DEFINE_LOG_CATEGORY_STATIC(LogIOSDeviceHelper, Log, All);
         FString DeviceName;
         FString DeviceID;
         FString DeviceType;
+		bool isAuthorized;
         DeviceConnectionInterface DeviceInterface;
     };
 
@@ -103,14 +108,44 @@ DEFINE_LOG_CATEGORY_STATIC(LogIOSDeviceHelper, Log, All);
             }
 
             FPlatformProcess::ExecProcess(*LibimobileDeviceInfo, *Arguments, &ReturnCodeInfo, &OutStdOutInfo, &OutStdErrInfo, NULL, true);
-            // parse product type and device name
-            FString DeviceName;
+            
+			LibIMobileDevice ToAdd;
+
+            // ideviceinfo can fail when the connected device is untrusted. It can be "Pairing dialog response pending (-19)", "Invalid HostID (-21)" or "User denied pairing (-18)"
+            // the only thing we can do is to make sure the Trust popup is correctly displayed.
+            if (OutStdErrInfo.Contains("ERROR: "))
+            {
+				if (OutStdErrInfo.Contains("Could not connect to lockdownd"))
+				{
+					UE_LOG(LogIOSDeviceHelper, Warning, TEXT("Could not pair with connected iOS/tvOS device. Trust this computer by accepting the popup on device."));
+					FString LibimobileDevicePair = GetLibImobileDeviceExe("idevicepair");
+					FString PairArguments = "-u " + DeviceID + " pair";
+					FPlatformProcess::ExecProcess(*LibimobileDevicePair, *PairArguments, &ReturnCodeInfo, &OutStdOutInfo, &OutStdErrInfo, NULL, true);
+				}
+				else
+				{
+					UE_LOG(LogIOSDeviceHelper, Warning, TEXT("Libimobile call failed : %s"), *OutStdErrInfo);
+				}
+				OutStdOutInfo.Empty();
+				OutStdErrInfo.Empty();
+				ToAdd.isAuthorized = false;
+			}
+			else
+			{
+				ToAdd.isAuthorized = true;
+			}
+
+			// parse product type and device name
+			FString DeviceName;
             OutStdOutInfo.Split(TEXT("DeviceName: "), nullptr, &DeviceName, ESearchCase::CaseSensitive, ESearchDir::FromStart);
-            DeviceName.Split(LINE_TERMINATOR, &DeviceName, nullptr, ESearchCase::CaseSensitive, ESearchDir::FromStart);
+			DeviceName.Split(LINE_TERMINATOR, &DeviceName, nullptr, ESearchCase::CaseSensitive, ESearchDir::FromStart);
+			if (!ToAdd.isAuthorized)
+            {
+				DeviceName = LOCTEXT("IosTvosUnauthorizedDevice", "iOS / tvOS (Unauthorized)").ToString();
+	        }
             FString ProductType;
             OutStdOutInfo.Split(TEXT("ProductType: "), nullptr, &ProductType, ESearchCase::CaseSensitive, ESearchDir::FromStart);
             ProductType.Split(LINE_TERMINATOR, &ProductType, nullptr, ESearchCase::CaseSensitive, ESearchDir::FromStart);
-            LibIMobileDevice ToAdd;
             ToAdd.DeviceID = DeviceID;
             ToAdd.DeviceName = DeviceName;
             ToAdd.DeviceType = ProductType;
@@ -260,7 +295,30 @@ public:
     }
 
 private:
+    
+    void NotifyDeviceChange(LibIMobileDevice& Device, bool bAdd)
+    {
+        FDeviceNotificationCallbackInformation CallbackInfo;
+
+        if (bAdd)
+        {
+            CallbackInfo.DeviceName = Device.DeviceName;
+            CallbackInfo.UDID = Device.DeviceID;
+            CallbackInfo.DeviceInterface = Device.DeviceInterface;
+            CallbackInfo.ProductType = Device.DeviceType;
+            CallbackInfo.msgType = 1;
+			CallbackInfo.isAuthorized = Device.isAuthorized;
+        }
+        else
+        {
+            CallbackInfo.UDID = Device.DeviceID;
+            CallbackInfo.msgType = 2;
+            DeviceNotification.Broadcast(&CallbackInfo);
 
+        }
+        DeviceNotification.Broadcast(&CallbackInfo);
+    }
+    
     void QueryDevices()
     {
         FString OutStdOut;
@@ -277,6 +335,11 @@ private:
                 UE_LOG(LogIOSDeviceHelper, Verbose, TEXT("IOS device listing is disabled for 1 minute (too many failed attempts)!"));
                 Enable(false);
             }
+			for (LibIMobileDevice device : CachedDevices)
+			{
+				NotifyDeviceChange(device, false);
+			}
+			CachedDevices.Empty();
             return;
         }
         RetryQuery = 5;
@@ -285,43 +348,43 @@ private:
 
         for (int32 Index = 0; Index < ParsedDevices.Num(); ++Index)
         {
-            // move on to next device if this one is already a known device
-            if (ConnectedDeviceIds.Find(ParsedDevices[Index].DeviceID) != INDEX_NONE)
-            {
-                ConnectedDeviceIds.Remove(ParsedDevices[Index].DeviceID);
-                continue;
-            }
-
-            // create an FIOSDevice
-            FDeviceNotificationCallbackInformation CallbackInfo;
-            CallbackInfo.DeviceName = ParsedDevices[Index].DeviceName;
-            CallbackInfo.UDID = ParsedDevices[Index].DeviceID;
-            CallbackInfo.DeviceInterface = ParsedDevices[Index].DeviceInterface;
-            CallbackInfo.ProductType = ParsedDevices[Index].DeviceType;
-            CallbackInfo.msgType = 1;
-            DeviceNotification.Broadcast(&CallbackInfo);
+			TObjectPtr<LibIMobileDevice> Found = CachedDevices.FindByPredicate(
+				[&](LibIMobileDevice Element) {
+					return Element.DeviceID == ParsedDevices[Index].DeviceID;
+				});
+			if (Found != nullptr)
+			{
+				if (Found->isAuthorized != ParsedDevices[Index].isAuthorized)
+				{
+					NotifyDeviceChange(ParsedDevices[Index], false);
+					NotifyDeviceChange(ParsedDevices[Index], true);
+				}
+				Found->DeviceID = "DealtWith";
+			}
+			else
+			{
+				NotifyDeviceChange(ParsedDevices[Index], true);
+			}
+	
         }
-        
-        // remove all devices no longer found
-        for (int32 DeviceIndex = 0; DeviceIndex < ConnectedDeviceIds.Num(); ++DeviceIndex)
-        {
-            FDeviceNotificationCallbackInformation CallbackInfo;
-            CallbackInfo.UDID = ConnectedDeviceIds[DeviceIndex];
-            CallbackInfo.msgType = 2;
-            DeviceNotification.Broadcast(&CallbackInfo);
-        }
-        ConnectedDeviceIds.Empty();
-        for (int32 Index = 0; Index < ParsedDevices.Num(); ++Index)
+
+        for (int32 Index = 0; Index < CachedDevices.Num(); ++Index)
         {
-            ConnectedDeviceIds.Add(ParsedDevices[Index].DeviceID);
+            if (CachedDevices[Index].DeviceID != "DealtWith")
+            {
+                NotifyDeviceChange(CachedDevices[Index], false);
+            }
         }
+
+        CachedDevices.Empty();
+		CachedDevices = ParsedDevices;
     }
 
     bool Stopping;
     bool bCheckDevices;
     bool NeedSDKCheck;
     int RetryQuery;
-    TArray<FString> ConnectedDeviceIds;
+    TArray<LibIMobileDevice> CachedDevices;
     FDeviceNotification DeviceNotification;
 };
 
@@ -403,6 +466,7 @@ void FIOSDeviceHelper::DoDeviceConnect(void* CallbackInfo)
     Event.DeviceID = FString::Printf(TEXT("%s@%s"), cbi->ProductType.Contains(TEXT("AppleTV")) ? TEXT("TVOS") : TEXT("IOS"), *(cbi->UDID));
     Event.DeviceName = cbi->DeviceName;
     Event.DeviceType = cbi->ProductType;
+	Event.bIsAuthorized = cbi->isAuthorized;
     Event.bCanReboot = false;
     Event.bCanPowerOn = false;
     Event.bCanPowerOff = false;
diff --git a/Engine/Source/Developer/IOS/IOSTargetPlatform/Private/IOSTargetDevice.h b/Engine/Source/Developer/IOS/IOSTargetPlatform/Private/IOSTargetDevice.h
--- a/Engine/Source/Developer/IOS/IOSTargetPlatform/Private/IOSTargetDevice.h
+++ b/Engine/Source/Developer/IOS/IOSTargetPlatform/Private/IOSTargetDevice.h
@@ -120,6 +120,7 @@ public:
 	virtual bool IsDefault() const override;
 	virtual bool PowerOff(bool Force) override;
 	virtual bool PowerOn() override;
+	virtual bool IsAuthorized() const override { return bIsDeviceAuthorized; }
 	virtual bool Reboot(bool bReconnect = false) override;
 	virtual bool SupportsFeature(ETargetDeviceFeatures Feature) const;
 	virtual bool TerminateProcess(const int64 ProcessId) override;
@@ -173,6 +174,9 @@ private:
 	/** Name of device */
 	FString DeviceName;
 
+	// Holds a flag indicating whether the device is USB / OTA comms authorized
+	bool bIsDeviceAuthorized;
+
 	/** Type of device */
 	ETargetDeviceTypes DeviceType;
 
@@ -206,6 +210,16 @@ public:
 		DeviceName = InDeviceName;
 	}
 
+	/**
+ * Sets the device's authorization state.
+ *
+ * @param bInIsAuthorized - Whether the device is authorized for USB communications.
+ */
+	void SetAuthorized(bool bInIsAuthorized)
+	{
+		bIsDeviceAuthorized = bInIsAuthorized;
+	}
+
 	/** Sets the type of the device */
 	void SetDeviceType(const FString InDeviceTypeString)
 	{
diff --git a/Engine/Source/Developer/IOS/IOSTargetPlatform/Private/IOSTargetPlatform.cpp b/Engine/Source/Developer/IOS/IOSTargetPlatform/Private/IOSTargetPlatform.cpp
--- a/Engine/Source/Developer/IOS/IOSTargetPlatform/Private/IOSTargetPlatform.cpp
+++ b/Engine/Source/Developer/IOS/IOSTargetPlatform/Private/IOSTargetPlatform.cpp
@@ -376,7 +376,7 @@ void FIOSTargetPlatform::HandleDeviceConnected(const FIOSLaunchDaemonPong& Messa
 	
 	if (!Device.IsValid())
 	{
-		if ((Message.DeviceType.Contains(TEXT("AppleTV")) && bIsTVOS) || (!Message.DeviceType.Contains(TEXT("AppleTV")) && !bIsTVOS))
+	if ((Message.DeviceType.Contains(TEXT("AppleTV")) && bIsTVOS) || (!Message.DeviceType.Contains(TEXT("AppleTV")) && !bIsTVOS))
 		{
 			Device = MakeShareable(new FIOSTargetDevice(*this));
 
@@ -385,6 +385,7 @@ void FIOSTargetPlatform::HandleDeviceConnected(const FIOSLaunchDaemonPong& Messa
 			Device->SetFeature(ETargetDeviceFeatures::PowerOff, Message.bCanPowerOff);
 			Device->SetDeviceId(DeviceId);
 			Device->SetDeviceName(Message.DeviceName);
+			Device->SetAuthorized(Message.bIsAuthorized);
 			Device->SetDeviceType(Message.DeviceType);
 			Device->SetIsSimulated(Message.DeviceID.Contains(TEXT("Simulator")));
 
diff --git a/Engine/Source/Runtime/IOS/LaunchDaemonMessages/Classes/IOSMessageProtocol.h b/Engine/Source/Runtime/IOS/LaunchDaemonMessages/Classes/IOSMessageProtocol.h
--- a/Engine/Source/Runtime/IOS/LaunchDaemonMessages/Classes/IOSMessageProtocol.h
+++ b/Engine/Source/Runtime/IOS/LaunchDaemonMessages/Classes/IOSMessageProtocol.h
@@ -40,6 +40,9 @@ struct FIOSLaunchDaemonPong
 	UPROPERTY()
 	bool bCanPowerOn;
 
+	UPROPERTY()
+	bool bIsAuthorized;
+
 	UPROPERTY()
 	bool bCanReboot;
 
@@ -49,13 +52,14 @@ struct FIOSLaunchDaemonPong
 		, bCanReboot(false)
 	{}
 
-	FIOSLaunchDaemonPong(FString InDeviceID, FString InDeviceName, FString InDeviceStatus, FString InDeviceType, bool bInCanPowerOff, bool bInCanPowerOn, bool bInCanReboot)
+	FIOSLaunchDaemonPong(FString InDeviceID, FString InDeviceName, FString InDeviceStatus, FString InDeviceType, bool bInCanPowerOff, bool bInCanPowerOn, bool bInCanReboot, bool bIsAuthorized)
 		: DeviceID(InDeviceID)
 		, DeviceName(InDeviceName)
 		, DeviceStatus(InDeviceStatus)
 		, DeviceType(InDeviceType)
 		, bCanPowerOff(bInCanPowerOff)
 		, bCanPowerOn(bInCanPowerOn)
+		, bIsAuthorized(bIsAuthorized)
 		, bCanReboot(bInCanReboot)
 	{}
 };
diff --git a/Engine/Source/Runtime/Launch/Private/IOS/GameLaunchDaemonMessageHandler.cpp b/Engine/Source/Runtime/Launch/Private/IOS/GameLaunchDaemonMessageHandler.cpp
--- a/Engine/Source/Runtime/Launch/Private/IOS/GameLaunchDaemonMessageHandler.cpp
+++ b/Engine/Source/Runtime/Launch/Private/IOS/GameLaunchDaemonMessageHandler.cpp
@@ -49,6 +49,7 @@ void FGameLaunchDaemonMessageHandler::HandlePingMessage(const FIOSLaunchDaemonPi
 			, false
 			, false
 			, false
+			, true
 			), MessageSender);
 	}
 }
