Fix Builds Tools Mac compiler error

diff --git a/Engine/Source/Developer/IOS/IOSTargetPlatform/Private/IOSDeviceHelper.cpp b/Engine/Source/Developer/IOS/IOSTargetPlatform/Private/IOSDeviceHelper.cpp
--- a/Engine/Source/Developer/IOS/IOSTargetPlatform/Private/IOSDeviceHelper.cpp
+++ b/Engine/Source/Developer/IOS/IOSTargetPlatform/Private/IOSDeviceHelper.cpp
@@ -510,3 +510,5 @@ void FIOSDeviceHelper::EnableDeviceCheck(bool OnOff)
 {
     QueryTask->Enable(OnOff);
 }
+
+#undef LOCTEXT_NAMESPACE
