Support for Device Output Window iOS log and console command sending using libimobiledevice

diff --git a/Engine/Source/Developer/IOS/IOSTargetPlatform/Private/IOSDeviceHelper.cpp b/Engine/Source/Developer/IOS/IOSTargetPlatform/Private/IOSDeviceHelper.cpp
--- a/Engine/Source/Developer/IOS/IOSTargetPlatform/Private/IOSDeviceHelper.cpp
+++ b/Engine/Source/Developer/IOS/IOSTargetPlatform/Private/IOSDeviceHelper.cpp
@@ -43,24 +43,6 @@ DEFINE_LOG_CATEGORY_STATIC(LogIOSDeviceHelper, Log, All);
         DeviceConnectionInterface DeviceInterface;
     };
 
-    static FString GetLibImobileDeviceExe(const FString &ExeName)
-    {
-        FString ToReturn;
-#if PLATFORM_WINDOWS
-        ToReturn = FPaths::ConvertRelativePathToFull(FPaths::EngineDir() / TEXT("Extras/ThirdPartyNotUE/libimobiledevice/x64/"));
-#elif PLATFORM_MAC
-        ToReturn = FPaths::ConvertRelativePathToFull(FPaths::EngineDir() / TEXT("Extras/ThirdPartyNotUE/libimobiledevice/Mac/"));
-#else
-        UE_LOG(LogIOSDeviceHelper, Error, TEXT("The current platform is unsupported by Libimobile library."));
-#endif
-        ToReturn += ExeName;
-#if PLATFORM_WINDOWS
-		ToReturn += TEXT(".exe");
-#endif
-
-        return ToReturn;
-    }
-
     static TArray<LibIMobileDevice> GetLibIMobileDevices()
     {
         FString OutStdOut;
diff --git a/Engine/Source/Developer/IOS/IOSTargetPlatform/Private/IOSTargetDevice.cpp b/Engine/Source/Developer/IOS/IOSTargetPlatform/Private/IOSTargetDevice.cpp
--- a/Engine/Source/Developer/IOS/IOSTargetPlatform/Private/IOSTargetDevice.cpp
+++ b/Engine/Source/Developer/IOS/IOSTargetPlatform/Private/IOSTargetDevice.cpp
@@ -139,7 +139,27 @@ bool FIOSTargetDevice::GetUserCredentials(FString& OutUserName, FString& OutUser
 
 inline void FIOSTargetDevice::ExecuteConsoleCommand(const FString& ExecCommand) const
 {
-	// this is where the new console command system execution goes.
+	FString OutStdOut;
+	FString OutStdErr;
+	FString Exe = GetLibImobileDeviceExe("itcpconnect");
+	FString Params = FString::Printf(TEXT(" -u %s 8888"), *DeviceId.GetDeviceName());
+
+	void* StdInPipe_Read = nullptr;
+	void* StdInPipe_Write = nullptr;
+
+	verify(FPlatformProcess::CreatePipe(StdInPipe_Read, StdInPipe_Write, true));
+
+	FProcHandle ProcHandle = FPlatformProcess::CreateProc(*Exe, *Params, true, false, false, nullptr, 0, nullptr, nullptr, StdInPipe_Read);
+	if (ProcHandle.IsValid())
+	{
+		FPlatformProcess::WritePipe(StdInPipe_Write, ExecCommand, nullptr);
+		FPlatformProcess::ClosePipe(StdInPipe_Read, StdInPipe_Write);
+		FPlatformProcess::WaitForProc(ProcHandle);
+	}
+	else
+	{
+		FPlatformProcess::ClosePipe(StdInPipe_Read, StdInPipe_Write);
+	}
 }
 
 inline ITargetDeviceOutputPtr FIOSTargetDevice::CreateDeviceOutputRouter(FOutputDevice* Output) const
diff --git a/Engine/Source/Developer/IOS/IOSTargetPlatform/Private/IOSTargetDeviceOutput.h b/Engine/Source/Developer/IOS/IOSTargetPlatform/Private/IOSTargetDeviceOutput.h
--- a/Engine/Source/Developer/IOS/IOSTargetPlatform/Private/IOSTargetDeviceOutput.h
+++ b/Engine/Source/Developer/IOS/IOSTargetPlatform/Private/IOSTargetDeviceOutput.h
@@ -12,10 +12,28 @@
 
 class FIOSTargetDevice;
 
+static FString GetLibImobileDeviceExe(const FString& ExeName)
+{
+	FString ToReturn;
+#if PLATFORM_WINDOWS
+	ToReturn = FPaths::ConvertRelativePathToFull(FPaths::EngineDir() / TEXT("Extras/ThirdPartyNotUE/libimobiledevice/x64/"));
+#elif PLATFORM_MAC
+	ToReturn = FPaths::ConvertRelativePathToFull(FPaths::EngineDir() / TEXT("Extras/ThirdPartyNotUE/libimobiledevice/Mac/"));
+#else
+	UE_LOG(LogIOSDeviceHelper, Error, TEXT("The current platform is unsupported by Libimobile library."));
+#endif
+	ToReturn += ExeName;
+#if PLATFORM_WINDOWS
+	ToReturn += TEXT(".exe");
+#endif
+
+	return ToReturn;
+}
+
 class FIOSDeviceOutputReaderRunnable : public FRunnable
 {
 public:
-	FIOSDeviceOutputReaderRunnable(const FTargetDeviceId InDeviceId, FOutputDevice* Output);
+	FIOSDeviceOutputReaderRunnable(const FString& InDeviceUDID, FOutputDevice* Output);
 	
 	// FRunnable interface.
 	virtual bool Init(void) override;
@@ -24,12 +42,16 @@ public:
 	virtual uint32 Run(void) override;
 
 private:
+	bool StartSyslogProcess(void);
+
 	// > 0 if we've been asked to abort work in progress at the next opportunity
 	FThreadSafeCounter	StopTaskCounter;
 	
-	FTargetDeviceId		DeviceId;
+	FString				DeviceUDID;
 	FOutputDevice*		Output;
-	TQueue<FString>		OutputQueue;
+	void*				SyslogReadPipe;
+	void*				SyslogWritePipe;
+	FProcHandle			SyslogProcHandle;
 };
 
 /**
@@ -42,7 +64,6 @@ public:
 	
 private:
 	TUniquePtr<FRunnableThread>						DeviceOutputThread;
-	FTargetDeviceId									DeviceId;
 	FString											DeviceName;
 };
 
diff --git a/Engine/Source/Developer/IOS/IOSTargetPlatform/Private/IOSTargetDeviceOutput.inl b/Engine/Source/Developer/IOS/IOSTargetPlatform/Private/IOSTargetDeviceOutput.inl
--- a/Engine/Source/Developer/IOS/IOSTargetPlatform/Private/IOSTargetDeviceOutput.inl
+++ b/Engine/Source/Developer/IOS/IOSTargetPlatform/Private/IOSTargetDeviceOutput.inl
@@ -10,25 +10,36 @@ class FIOSDeviceOutputReaderRunnable;
 class FIOSTargetDevice;
 class FIOSTargetDeviceOutput;
 
-
-
-
-inline FIOSDeviceOutputReaderRunnable::FIOSDeviceOutputReaderRunnable(const FTargetDeviceId InDeviceId, FOutputDevice* InOutput)
+inline FIOSDeviceOutputReaderRunnable::FIOSDeviceOutputReaderRunnable(const FString& InDeviceUDID, FOutputDevice* InOutput)
 	: StopTaskCounter(0)
-	, DeviceId(InDeviceId)
+	, DeviceUDID(InDeviceUDID)
 	, Output(InOutput)
+	, SyslogReadPipe(nullptr)
+	, SyslogWritePipe(nullptr)
 {
 }
 
+inline bool FIOSDeviceOutputReaderRunnable::StartSyslogProcess(void)
+{
+	FString Exe = GetLibImobileDeviceExe("idevicesyslog");
+	FString Params = FString::Printf(TEXT(" -u %s -m [UE]"), *DeviceUDID);
+	SyslogProcHandle = FPlatformProcess::CreateProc(*Exe, *Params, true, false, false, NULL, 0, NULL, SyslogWritePipe);
+	return SyslogProcHandle.IsValid();
+}
 
-inline bool FIOSDeviceOutputReaderRunnable::Init(void) 
-{ 
-	return true;
+inline bool FIOSDeviceOutputReaderRunnable::Init(void)
+{
+	FPlatformProcess::CreatePipe(SyslogReadPipe, SyslogWritePipe);
+	return StartSyslogProcess();
 }
 
-inline void FIOSDeviceOutputReaderRunnable::Exit(void) 
+inline void FIOSDeviceOutputReaderRunnable::Exit(void)
 {
-	StopTaskCounter.Increment();
+	if (SyslogProcHandle.IsValid())
+	{
+		FPlatformProcess::TerminateProc(SyslogProcHandle);
+	}
+	FPlatformProcess::ClosePipe(SyslogReadPipe, SyslogWritePipe);
 }
 
 inline void FIOSDeviceOutputReaderRunnable::Stop(void)
@@ -38,19 +49,52 @@ inline void FIOSDeviceOutputReaderRunnable::Stop(void)
 
 inline uint32 FIOSDeviceOutputReaderRunnable::Run(void)
 {
-	Output->Serialize(TEXT("Starting Output"), ELogVerbosity::Log, NAME_None);
-	while (StopTaskCounter.GetValue() == 0)
+	FString SyslogOutput;
+
+	while (StopTaskCounter.GetValue() == 0 && SyslogProcHandle.IsValid())
 	{
-		FString Text;
-		if (OutputQueue.Dequeue(Text))
+		if (!FPlatformProcess::IsProcRunning(SyslogProcHandle))
 		{
-			if (Text.Contains(TEXT("[UE]"), ESearchCase::CaseSensitive))
+			// When user plugs out USB cable idevicesyslog process stops
+			// Keep trying to restore idevicesyslog connection until code that uses this object will not kill us
+			Output->Serialize(TEXT("Trying to restore connection to device..."), ELogVerbosity::Log, NAME_None);
+			FPlatformProcess::CloseProc(SyslogProcHandle);
+			if (StartSyslogProcess())
 			{
-				Output->Serialize(*Text, ELogVerbosity::Log, NAME_None);
+				FPlatformProcess::Sleep(1.0f);
+			}
+			else
+			{
+				Output->Serialize(TEXT("Failed to start idevicesyslog proccess"), ELogVerbosity::Log, NAME_None);
+				Stop();
 			}
 		}
 		else
 		{
+			SyslogOutput.Append(FPlatformProcess::ReadPipe(SyslogReadPipe));
+
+			if (SyslogOutput.Len() > 0)
+			{
+				TArray<FString> OutputLines;
+				SyslogOutput.ParseIntoArray(OutputLines, TEXT("\n"), false);
+
+				if (!SyslogOutput.EndsWith(TEXT("\n")))
+				{
+					// partial line at the end, do not serialize it until we receive remainder
+					SyslogOutput = OutputLines.Last();
+					OutputLines.RemoveAt(OutputLines.Num() - 1);
+				}
+				else
+				{
+					SyslogOutput.Reset();
+				}
+
+				for (int32 i = 0; i < OutputLines.Num(); ++i)
+				{
+					Output->Serialize(*OutputLines[i], ELogVerbosity::Log, NAME_None);
+				}
+			}
+
 			FPlatformProcess::Sleep(0.1f);
 		}
 	}
@@ -63,13 +107,9 @@ inline bool FIOSTargetDeviceOutput::Init(const FIOSTargetDevice& TargetDevice, F
 	check(Output);
 	// Output will be produced by background thread
 	check(Output->CanBeUsedOnAnyThread());
-	
-	DeviceId = TargetDevice.GetId();
 	DeviceName = TargetDevice.GetName();
-
-	Output->Serialize(TEXT("Creating FIOSTargetDeviceOutput ....."), ELogVerbosity::Log, NAME_None);
-		
-	auto* Runnable = new FIOSDeviceOutputReaderRunnable(DeviceId, Output);
+	auto* Runnable = new FIOSDeviceOutputReaderRunnable(TargetDevice.GetId().GetDeviceName(), Output);
 	DeviceOutputThread = TUniquePtr<FRunnableThread>(FRunnableThread::Create(Runnable, TEXT("FIOSDeviceOutputReaderRunnable")));
 	return true;
 }
+
diff --git a/Engine/Source/Runtime/Launch/Private/IOS/LaunchIOS.cpp b/Engine/Source/Runtime/Launch/Private/IOS/LaunchIOS.cpp
--- a/Engine/Source/Runtime/Launch/Private/IOS/LaunchIOS.cpp
+++ b/Engine/Source/Runtime/Launch/Private/IOS/LaunchIOS.cpp
@@ -28,6 +28,7 @@
 #include "Interfaces/IPv4/IPv4Address.h"
 #include "Interfaces/IPv4/IPv4Endpoint.h"
 #include "Misc/EmbeddedCommunication.h"
+#include "Misc/CoreDelegates.h"
 
 FEngineLoop GEngineLoop;
 FGameLaunchDaemonMessageHandler GCommandSystem;
@@ -444,12 +445,30 @@ void FAppEntry::Init()
 	// start up the engine
 	GEngineLoop.Init();
 #if !UE_BUILD_SHIPPING
-	UE_LOG(LogInit, Display, TEXT("Initializing TCPConsoleListener."));
+	FIPv4Endpoint ConsoleTCP(FIPv4Address::InternalLoopback, 8888); //TODO: read this from an .ini
 	if (ConsoleListener == nullptr)
 	{
-		FIPv4Endpoint ConsoleTCP(FIPv4Address::InternalLoopback, 8888); //TODO: @csulea read this from some .ini
 		ConsoleListener = new TcpConsoleListener(ConsoleTCP);
 	}
+	// tear down the console listener when backgrounded
+	FCoreDelegates::ApplicationWillEnterBackgroundDelegate.AddLambda([]()
+	{
+		if (ConsoleListener)
+		{
+			delete ConsoleListener;
+			ConsoleListener = nullptr;
+		}
+	});
+	FCoreDelegates::ApplicationHasEnteredForegroundDelegate.AddLambda([ConsoleTCP]()
+	{
+		if (ConsoleListener == nullptr)
+		{
+			ConsoleListener = new TcpConsoleListener(ConsoleTCP);
+		}
+	});
+
+
+
 #endif // UE_BUILD_SHIPPING
 }
 
diff --git a/Engine/Source/Runtime/Launch/Private/IOS/TcpConsoleListener.cpp b/Engine/Source/Runtime/Launch/Private/IOS/TcpConsoleListener.cpp
--- a/Engine/Source/Runtime/Launch/Private/IOS/TcpConsoleListener.cpp
+++ b/Engine/Source/Runtime/Launch/Private/IOS/TcpConsoleListener.cpp
@@ -5,8 +5,8 @@
 #include "Common/TcpSocketBuilder.h"
 #include "Common/TcpListener.h"
 
-extern void EnqueueConsoleCommand(uint8 *Command);
-TcpConsoleListener *ConsoleListener = nullptr;
+extern void EnqueueConsoleCommand(uint8* Command);
+TcpConsoleListener* ConsoleListener = nullptr;
 
 /* TcpConsoleListener structors
  *****************************************************************************/
@@ -16,13 +16,13 @@ TcpConsoleListener::TcpConsoleListener(const FIPv4Endpoint& InListenEndpoint)
 	, bStopping(false)
 	, Listener(nullptr)
 {
-	UE_LOG(LogTemp, Warning, TEXT("[UE] Console Listener created!\n"));
+	UE_LOG(LogTemp, Display, TEXT("Initializing TCPConsoleListener."));
 	Thread = FRunnableThread::Create(this, TEXT("TcpConsoleListener"), 128 * 1024, TPri_Normal);
 }
 
-
 TcpConsoleListener::~TcpConsoleListener()
 {
+	UE_LOG(LogTemp, Display, TEXT("Stopping TCPConsoleListener."));
 	if (Listener)
 	{
 		delete Listener;
@@ -39,7 +39,6 @@ TcpConsoleListener::~TcpConsoleListener()
 bool TcpConsoleListener::HandleListenerConnectionAccepted(FSocket* ClientSocket, const FIPv4Endpoint& ClientEndpoint)
 {
 	Connections.Add(ClientSocket);
-	
 	return true;
 }
 
@@ -54,7 +53,6 @@ void TcpConsoleListener::Exit()
 
 bool TcpConsoleListener::Init()
 {
-	UE_LOG(LogTemp, Warning, TEXT("[UE] TCP Listener created!\n"));
 	Listener = new FTcpListener(ListenEndpoint);
 	Listener->OnConnectionAccepted().BindRaw(this, &TcpConsoleListener::HandleListenerConnectionAccepted);
 
@@ -82,7 +80,7 @@ uint32 TcpConsoleListener::Run()
 				memset(RecvBuffer, 0, CommandSize);
 				if (Connection->Recv(RecvBuffer, CommandSize, BytesRead))
 				{
-					UE_LOG(LogTemp, Warning, TEXT("[UE] Got TCP console command size %i '%s'"), BytesRead, *FString(UTF8_TO_TCHAR(RecvBuffer)));
+					UE_LOG(LogTemp, Display, TEXT("Received TCP console command '%s'"), *FString(UTF8_TO_TCHAR(RecvBuffer)));
 					EnqueueConsoleCommand(RecvBuffer);
 				}
 				else
@@ -103,7 +101,6 @@ uint32 TcpConsoleListener::Run()
 	return 0;
 }
 
-
 void TcpConsoleListener::Stop()
 {
 	bStopping = true;
@@ -114,5 +111,3 @@ void TcpConsoleListener::Stop()
 		Listener = nullptr;
 	}
 }
-
-
