Harmonize Metal version to use in iOS / Mac

diff --git a/Engine/Config/BaseEngine.ini b/Engine/Config/BaseEngine.ini
--- a/Engine/Config/BaseEngine.ini
+++ b/Engine/Config/BaseEngine.ini
@@ -2539,7 +2539,7 @@ bAutomaticSigning=false
 UseFastIntrinsics=False
 ForceFloats=False
 EnableMathOptimisations=True
-MaxShaderLanguageVersion=7
+MetalLanguageVersion=5
 bDisableMotionData=False
 bEnableAdvertisingIdentifier=True
 
@@ -2672,7 +2672,7 @@ bTarget32Bit=false
 +TargetedRHIs=SF_VULKAN_SM5
 
 [/Script/MacTargetPlatform.MacTargetSettings]
-MaxShaderLanguageVersion=7
+MetalLanguageVersion=5
 +TargetedRHIs=SF_METAL_SM5
 UseFastIntrinsics=False
 ForceFloats=False
diff --git a/Engine/Shaders/Public/Platform.ush b/Engine/Shaders/Public/Platform.ush
--- a/Engine/Shaders/Public/Platform.ush
+++ b/Engine/Shaders/Public/Platform.ush
@@ -260,21 +260,13 @@
 	// @todo metal: remove this and make sure all uses handle METAL_PROFILE
 	#undef ES3_1_PROFILE
 	#define ES3_1_PROFILE 1
-#if MAX_SHADER_LANGUAGE_VERSION >= 3
 	#define FCOLOR_COMPONENT_SWIZZLE .rgba
-#else
-	#define FCOLOR_COMPONENT_SWIZZLE .bgra
-#endif
 	#define FMANUALFETCH_COLOR_COMPONENT_SWIZZLE .bgra
 	#define STENCIL_COMPONENT_SWIZZLE .x
 
 #elif METAL_MRT_PROFILE
 	#define FEATURE_LEVEL FEATURE_LEVEL_SM5
-#if MAX_SHADER_LANGUAGE_VERSION >= 3
 	#define FCOLOR_COMPONENT_SWIZZLE .rgba
-#else
-	#define FCOLOR_COMPONENT_SWIZZLE .bgra
-#endif
 	#define FMANUALFETCH_COLOR_COMPONENT_SWIZZLE .bgra
 	#define STENCIL_COMPONENT_SWIZZLE .x
 
diff --git a/Engine/Shaders/Public/Platform/Metal/MetalCommon.ush b/Engine/Shaders/Public/Platform/Metal/MetalCommon.ush
--- a/Engine/Shaders/Public/Platform/Metal/MetalCommon.ush
+++ b/Engine/Shaders/Public/Platform/Metal/MetalCommon.ush
@@ -19,7 +19,7 @@
 #define StrongTypedBuffer Buffer
 
 /* Too many bugs to enable wave-intrinsics just yet, but we are getting there.
-#if ((MAX_SHADER_LANGUAGE_VERSION >= 3 && COMPUTESHADER) || (MAX_SHADER_LANGUAGE_VERSION >= 4 && PIXELSHADER))
+#if COMPUTESHADER || PIXELSHADER
 #define COMPILER_SUPPORTS_WAVE_ONCE 1
 #define COMPILER_SUPPORTS_WAVE_VOTE 1
 #define COMPILER_SUPPORTS_WAVE_MINMAX 1
diff --git a/Engine/Source/Developer/Apple/MetalShaderFormat/Private/MetalShaderCompiler.cpp b/Engine/Source/Developer/Apple/MetalShaderFormat/Private/MetalShaderCompiler.cpp
--- a/Engine/Source/Developer/Apple/MetalShaderFormat/Private/MetalShaderCompiler.cpp
+++ b/Engine/Source/Developer/Apple/MetalShaderFormat/Private/MetalShaderCompiler.cpp
@@ -75,7 +75,7 @@ static inline uint32 ParseNumber(const ANSICHAR* Str)
 
 struct FHlslccMetalHeader : public CrossCompiler::FHlslccHeader
 {
-	FHlslccMetalHeader(uint8 const Version);
+	FHlslccMetalHeader(uint32 const Version);
 	virtual ~FHlslccMetalHeader();
 	
 	// After the standard header, different backends can output their own info
@@ -83,10 +83,10 @@ struct FHlslccMetalHeader : public CrossCompiler::FHlslccHeader
 	
 	TMap<uint8, TArray<uint8>> ArgumentBuffers;
 	int8 SideTable;
-	uint8 Version;
+	uint32 Version;
 };
 
-FHlslccMetalHeader::FHlslccMetalHeader(uint8 const InVersion)
+FHlslccMetalHeader::FHlslccMetalHeader(uint32 const InVersion)
 {
 	SideTable = -1;
 	Version = InVersion;
@@ -845,8 +845,7 @@ void CompileShader_Metal(const FShaderCompilerInput& _Input,FShaderCompilerOutpu
 	
 	EMetalGPUSemantics Semantics = EMetalGPUSemanticsMobile;
 	
-	FString const* MaxVersion = Input.Environment.GetDefinitions().Find(TEXT("MAX_SHADER_LANGUAGE_VERSION"));
-	int32 VersionEnum = FMath::Min(5, FCString::Atoi(*FString(*MaxVersion)));
+    int32 VersionEnum = FCString::Atoi(*FString(*Input.Environment.GetDefinitions().Find(TEXT("SHADER_LANGUAGE_VERSION"))));
 
 	// TODO read from toolchain
 	bool bAppleTV = (Input.ShaderFormat == NAME_SF_METAL_TVOS || Input.ShaderFormat == NAME_SF_METAL_MRT_TVOS);
@@ -925,6 +924,7 @@ void CompileShader_Metal(const FShaderCompilerInput& _Input,FShaderCompilerOutpu
 		break;
 
 	case 5:
+    case 0:
 		TypeMode = EMetalTypeBufferModeTB;
 		StandardVersion = TEXT("2.2");
 		if (bAppleTV)
@@ -945,7 +945,7 @@ void CompileShader_Metal(const FShaderCompilerInput& _Input,FShaderCompilerOutpu
 		Output.bSucceeded = false;
 		{
 			FShaderCompilerError* NewError = new(Output.Errors) FShaderCompilerError();
-			NewError->StrippedErrorMessage = FString::Printf(TEXT("Minimum Version is 2.2 in UE."), *StandardVersion);
+			NewError->StrippedErrorMessage = FString::Printf(TEXT("Minimum Metal Version is 2.2 in UE5.0"));
 			return;
 		}
 		break;
diff --git a/Engine/Source/Developer/IOS/IOSPlatformEditor/Private/IOSTargetSettingsCustomization.cpp b/Engine/Source/Developer/IOS/IOSPlatformEditor/Private/IOSTargetSettingsCustomization.cpp
--- a/Engine/Source/Developer/IOS/IOSPlatformEditor/Private/IOSTargetSettingsCustomization.cpp
+++ b/Engine/Source/Developer/IOS/IOSPlatformEditor/Private/IOSTargetSettingsCustomization.cpp
@@ -818,7 +818,7 @@ void FIOSTargetSettingsCustomization::BuildPListSection(IDetailLayoutBuilder& De
 
     // Handle max. shader version a little specially.
     {
-        ShaderVersionPropertyHandle = DetailLayout.GetProperty(GET_MEMBER_NAME_CHECKED(UIOSRuntimeSettings, MaxShaderLanguageVersion));
+        ShaderVersionPropertyHandle = DetailLayout.GetProperty(GET_MEMBER_NAME_CHECKED(UIOSRuntimeSettings, MetalLanguageVersion));
 		ShaderVersionPropertyHandle->SetOnPropertyValueChanged(OnUpdateShaderStandardWarning);
 		
 		// Drop-downs for setting type of lower and upper bound normalization
@@ -1877,38 +1877,32 @@ FText FIOSTargetSettingsCustomization::GetMinVersionDesc() const
 
 void FIOSTargetSettingsCustomization::SetShaderStandard(int32 Value)
 {
-	FPropertyAccess::Result Res = ShaderVersionPropertyHandle->SetValue((uint8)Value);
-	check(Res == FPropertyAccess::Success);
-	
-	if (MinOSPropertyHandle.IsValid())
-	{
-		FText Message;
-		
-		uint8 EnumValue = (uint8)EIOSVersion::IOS_13;
-		if (MinOSPropertyHandle.IsValid())
-		{
-			MinOSPropertyHandle->GetValue(EnumValue);
-		}
-		
-		bool bMRTEnabled = false;
-		if (MRTPropertyHandle.IsValid())
-		{
-			MRTPropertyHandle->GetValue(bMRTEnabled);
-		}
-		
-		// make sure we never set the min version to less than current supported
-		if (((EIOSVersion)EnumValue < EIOSVersion::IOS_13))
-		{
-			SetMinVersion((int32)EIOSVersion::IOS_13);
-		}
+    FPropertyAccess::Result Res = ShaderVersionPropertyHandle->SetValue((uint8)Value);
+    check(Res == FPropertyAccess::Success);
 
-		
-		ShaderVersionWarningTextBox->SetError(Message);
-	}
-	else
-	{
-		ShaderVersionWarningTextBox->SetError(TEXT(""));
-	}
+    if (MinOSPropertyHandle.IsValid())
+    {
+        uint8 IOSVersion = (uint8)EIOSVersion::IOS_13;
+        if (MinOSPropertyHandle.IsValid())
+        {
+            MinOSPropertyHandle->GetValue(IOSVersion);
+        }
+
+        ShaderVersionWarningTextBox->SetError(TEXT(""));
+
+        switch (IOSVersion)
+        {
+            case (uint8)EIOSVersion::IOS_13:
+                if (Value != 0 && Value < (int32)EIOSMetalShaderStandard::IOSMetalSLStandard_2_2){Value = (int32)EIOSMetalShaderStandard::IOSMetalSLStandard_2_2; ShaderVersionWarningTextBox->SetError(TEXT("Metal 2.2 is the minimum for iOS13")); return;}
+                break;;
+            case (uint8)EIOSVersion::IOS_14:
+                if (Value != 0 && Value < (int32)EIOSMetalShaderStandard::IOSMetalSLStandard_2_3){Value = (int32)EIOSMetalShaderStandard::IOSMetalSLStandard_2_3; ShaderVersionWarningTextBox->SetError(TEXT("Metal 2.3 is the minimum for iOS14")); return;}
+                break;
+            case (uint8)EIOSVersion::IOS_15:
+                if (Value != 0 && Value < (int32)EIOSMetalShaderStandard::IOSMetalSLStandard_2_4){Value = (int32)EIOSMetalShaderStandard::IOSMetalSLStandard_2_4; ShaderVersionWarningTextBox->SetError(TEXT("Metal 2.4 is the minimum for iOS15")); return;}
+                break;
+        }
+    }
 }
 
 void FIOSTargetSettingsCustomization::UpdateShaderStandardWarning()
@@ -1960,9 +1954,9 @@ void FIOSTargetSettingsCustomization::UpdateMetalMRTWarning()
 			}
 			
 			ShaderVersionPropertyHandle->GetValue(EnumValue);
-			if (EnumValue < (uint8)EIOSMetalShaderStandard::IOSMetalSLStandard_2_2)
+			if (EnumValue < (int32)EIOSMetalShaderStandard::IOSMetalSLStandard_Minimum)
 			{
-				SetShaderStandard((int32)EIOSMetalShaderStandard::IOSMetalSLStandard_2_2);
+				SetShaderStandard((int32)EIOSMetalShaderStandard::IOSMetalSLStandard_Minimum);
 			}
 		}
 		else
diff --git a/Engine/Source/Developer/Mac/MacTargetPlatform/Classes/MacTargetSettings.h b/Engine/Source/Developer/Mac/MacTargetPlatform/Classes/MacTargetSettings.h
--- a/Engine/Source/Developer/Mac/MacTargetPlatform/Classes/MacTargetSettings.h
+++ b/Engine/Source/Developer/Mac/MacTargetPlatform/Classes/MacTargetSettings.h
@@ -13,17 +13,14 @@
 UENUM()
 enum class EMacMetalShaderStandard : uint8
 {
-    /** Metal Shaders Compatible With macOS 10.13 or later (std=osx-metal2.0) */
-    MacMetalSLStandard_2_0 = 3 UMETA(DisplayName="Metal v2.0 (10.13.0+)"),
-    
-    /** Metal Shaders Compatible With macOS 10.14 or later (std=osx-metal2.1) */
-    MacMetalSLStandard_2_1 = 4 UMETA(DisplayName="Metal v2.1 (10.14.0+)"),
-
-	/** Metal Shaders Compatible With macOS 10.15 or later (std=osx-metal2.2) */
-	MacMetalSLStandard_2_2 = 5 UMETA(DisplayName="Metal v2.2 (10.15.0+)"),
-
-	/** Metal Shaders Compatible With macOS 11.0 or later (std=osx-metal2.3) */
-	MacMetalSLStandard_2_3 = 6 UMETA(DisplayName="Metal v2.3 (11.0+)"),
+    /** Metal Shader 2.2 is the minimum as of UE5.0*/
+    MacMetalSLStandard_Minimum = 0 UMETA(DisplayName="Minimum, Currently v2.2 (10.15.0+)"),
+    /** Metal Shaders Compatible With macOS 10.15 or later (std=osx-metal2.2) */
+    MacMetalSLStandard_2_2 = 5 UMETA(DisplayName="Metal v2.2 (10.15.0+)"),
+    /** Metal Shaders Compatible With macOS 11.0 or later (std=osx-metal2.3) */
+    MacMetalSLStandard_2_3 = 6 UMETA(DisplayName="Metal v2.3 (11.0+)"),
+    /** Metal Shaders Compatible With macOS 12.0 or later (std=osx-metal2.4) */
+    MacMetalSLStandard_2_4 = 7 UMETA(DisplayName="Metal v2.4 (12.0+)"),
 };
 
 UENUM()
@@ -66,11 +63,10 @@ public:
     EMacTargetArchitecture TargetArchitecture;
     
     /**
-     * The maximum supported Metal shader langauge version. 
-     * This defines what features may be used and OS versions supported.
+     * The Metal shader language version which will be used when compiling the shaders.
      */
-    UPROPERTY(EditAnywhere, config, Category=Rendering, meta = (DisplayName = "Max. Metal Shader Standard To Target", ConfigRestartRequired = true))
-    uint8 MaxShaderLanguageVersion;
+    UPROPERTY(EditAnywhere, config, Category=Rendering, meta = (DisplayName = "Metal Shader Standard To Target", ConfigRestartRequired = true))
+    int32 MetalLanguageVersion;
     
     /**
      * Whether to use the Metal shading language's "fast" intrinsics.
diff --git a/Engine/Source/Developer/Mac/MacTargetPlatform/Private/MacTargetPlatformModule.cpp b/Engine/Source/Developer/Mac/MacTargetPlatform/Private/MacTargetPlatformModule.cpp
--- a/Engine/Source/Developer/Mac/MacTargetPlatform/Private/MacTargetPlatformModule.cpp
+++ b/Engine/Source/Developer/Mac/MacTargetPlatform/Private/MacTargetPlatformModule.cpp
@@ -45,11 +45,12 @@ public:
 		
 		// We need to manually load the config properties here, as this module is loaded before the UObject system is setup to do this
         GConfig->GetArray(TEXT("/Script/MacTargetPlatform.MacTargetSettings"), TEXT("TargetedRHIs"), TargetSettings->TargetedRHIs, GEngineIni);
-		
-		int32 Value = 1;
-        GConfig->GetInt(TEXT("/Script/MacTargetPlatform.MacTargetSettings"), TEXT("MaxShaderLanguageVersion"), Value, GEngineIni);
-        TargetSettings->MaxShaderLanguageVersion = FMath::Max(Value, 7);
-		
+       
+        if (!GConfig->GetInt(TEXT("/Script/MacTargetPlatform.MacTargetSettings"), TEXT("MetalLanguageVersion"), TargetSettings->MetalLanguageVersion, GEngineIni))
+        {
+            TargetSettings->MetalLanguageVersion = 0;
+        }
+        
 		if (!GConfig->GetBool(TEXT("/Script/MacTargetPlatform.MacTargetSettings"), TEXT("UseFastIntrinsics"), TargetSettings->UseFastIntrinsics, GEngineIni))
 		{
 			TargetSettings->UseFastIntrinsics = false;
diff --git a/Engine/Source/Editor/DetailCustomizations/Private/MacTargetSettingsDetails.cpp b/Engine/Source/Editor/DetailCustomizations/Private/MacTargetSettingsDetails.cpp
--- a/Engine/Source/Editor/DetailCustomizations/Private/MacTargetSettingsDetails.cpp
+++ b/Engine/Source/Editor/DetailCustomizations/Private/MacTargetSettingsDetails.cpp
@@ -115,7 +115,7 @@ void FMacTargetSettingsDetails::CustomizeDetails( IDetailLayoutBuilder& DetailBu
     // Handle max. shader version a little specially.
     {
         IDetailCategoryBuilder& RenderCategory = DetailBuilder.EditCategory(TEXT("Rendering"));
-        ShaderVersionPropertyHandle = DetailBuilder.GetProperty(TEXT("MaxShaderLanguageVersion"));
+        ShaderVersionPropertyHandle = DetailBuilder.GetProperty(TEXT("MetalLanguageVersion"));
 		
 		// Drop-downs for setting type of lower and upper bound normalization
 		IDetailPropertyRow& ShaderVersionPropertyRow = RenderCategory.AddProperty(ShaderVersionPropertyHandle.ToSharedRef());
@@ -282,13 +282,9 @@ bool FMacTargetSettingsDetails::HandlePostExternalIconCopy(const FString& InChos
 }
 
 static uint32 GMacTargetSettingsMinOSVers[][3] = {
-	{10,11,6},
-	{10,11,6},
-	{10,12,6},
-	{10,13,0},
-	{10,14,0},
 	{10,15,0},
 	{11, 0,0},
+    {12, 0,0}
 };
 
 TSharedRef<SWidget> FMacTargetSettingsDetails::OnGetShaderVersionContent()
@@ -315,14 +311,14 @@ TSharedRef<SWidget> FMacTargetSettingsDetails::OnGetShaderVersionContent()
 
 FText FMacTargetSettingsDetails::GetShaderVersionDesc() const
 {
-	uint8 EnumValue;
-	ShaderVersionPropertyHandle->GetValue(EnumValue);
+    int32 EnumValue;
+    ShaderVersionPropertyHandle->GetValue(EnumValue);
 	
 	UEnum* Enum = FindObjectChecked<UEnum>(ANY_PACKAGE, TEXT("EMacMetalShaderStandard"), true);
 	
 	if (EnumValue < Enum->GetMaxEnumValue() && Enum->IsValidEnumValue(EnumValue))
 	{
-		return Enum->GetDisplayNameTextByValue(EnumValue);
+		return Enum->GetDisplayNameTextByValue((uint8)EnumValue);
 	}
 	
 	return FText::GetEmpty();
@@ -330,25 +326,32 @@ FText FMacTargetSettingsDetails::GetShaderVersionDesc() const
 
 void FMacTargetSettingsDetails::SetShaderStandard(int32 Value)
 {
-	FText Message;
-	
-	if (Value == 2)
-	{
-		Message = LOCTEXT("DeprecatedMacMetalShaderVersion1_2","Metal Shader Standard v1.2 required for macOS 10.12 Sierra is deprecated in 4.21 and will be removed in the next version.");
-	}
-	
-	ShaderVersionWarningTextBox->SetError(Message);
-	
-	FPropertyAccess::Result Res = ShaderVersionPropertyHandle->SetValue((uint8)Value);
-	check(Res == FPropertyAccess::Success);
+    if (ShaderVersionPropertyHandle->IsValidHandle())
+    {
+        FPropertyAccess::Result Res = ShaderVersionPropertyHandle->SetValue(Value);
+        check(Res == FPropertyAccess::Success);
+    }
+    
+    ShaderVersionWarningTextBox->SetError(TEXT(""));
+    if (Value < 5) // EMacMetalShaderStandard::MacMetalSLStandard_Minimum
+    {
+        ShaderVersionWarningTextBox->SetError(TEXT("Metal Shader Standard is 2.2 on UE5.0"));
+    }
 }
 
 void FMacTargetSettingsDetails::UpdateShaderStandardWarning()
 {
 	// Update the UI
-	uint8 EnumValue;
-	ShaderVersionPropertyHandle->GetValue(EnumValue);
-	SetShaderStandard(EnumValue);
+	uint8 EnumValue = 0;
+
+    if (ShaderVersionPropertyHandle->IsValidHandle())
+    {
+        ShaderVersionPropertyHandle->GetValue(EnumValue);
+        if (EnumValue < 5)
+        {
+            SetShaderStandard(5); // EMacMetalShaderStandard::MacMetalSLStandard_Minimum
+        }
+    }
 }
 
 #undef LOCTEXT_NAMESPACE
diff --git a/Engine/Source/Runtime/Apple/AGXRHI/Private/AGXCommandQueue.cpp b/Engine/Source/Runtime/Apple/AGXRHI/Private/AGXCommandQueue.cpp
--- a/Engine/Source/Runtime/Apple/AGXRHI/Private/AGXCommandQueue.cpp
+++ b/Engine/Source/Runtime/Apple/AGXRHI/Private/AGXCommandQueue.cpp
@@ -23,21 +23,15 @@ FAGXCommandQueue::FAGXCommandQueue(uint32 const MaxNumCommandBuffers /* = 0 */)
 : ParallelCommandLists(0)
 , RuntimeDebuggingLevel(EAGXDebugLevelOff)
 {
-	int32 MaxShaderVersion = 0;
-	int32 DefaultMaxShaderVersion = 5; // MSL v2.2
-	int32 MinShaderVersion = 5; // MSL v2.2
-
+    int32 MetalShaderVersion = 0;
+    
 #if PLATFORM_MAC
 	const TCHAR* const Settings = TEXT("/Script/MacTargetPlatform.MacTargetSettings");
 #else
 	const TCHAR* const Settings = TEXT("/Script/IOSRuntimeSettings.IOSRuntimeSettings");
 #endif
-    if(!GConfig->GetInt(Settings, TEXT("MaxShaderLanguageVersion"), MaxShaderVersion, GEngineIni))
-	{
-		MaxShaderVersion = DefaultMaxShaderVersion;
-	}
-	MaxShaderVersion = FMath::Max(MinShaderVersion, MaxShaderVersion);
-	AGXValidateVersion(MaxShaderVersion);
+    GConfig->GetInt(Settings, TEXT("MetalLanguageVersion"), MetalShaderVersion, GEngineIni);
+	AGXValidateVersion(MetalShaderVersion);
 
 	if(MaxNumCommandBuffers == 0)
 	{
diff --git a/Engine/Source/Runtime/Apple/MetalRHI/Private/MetalCommandQueue.cpp b/Engine/Source/Runtime/Apple/MetalRHI/Private/MetalCommandQueue.cpp
--- a/Engine/Source/Runtime/Apple/MetalRHI/Private/MetalCommandQueue.cpp
+++ b/Engine/Source/Runtime/Apple/MetalRHI/Private/MetalCommandQueue.cpp
@@ -26,25 +26,20 @@ FMetalCommandQueue::FMetalCommandQueue(mtlpp::Device InDevice, uint32 const MaxN
 , ParallelCommandLists(0)
 , RuntimeDebuggingLevel(EMetalDebugLevelOff)
 {
-	int32 MaxShaderVersion = 0;
 	int32 IndirectArgumentTier = 0;
-	int32 DefaultMaxShaderVersion = 5; // MSL v2.2
-	int32 MinShaderVersion = 5; // MSL v2.2
+    int32 MetalShaderVersion = 0;
 #if PLATFORM_MAC
 	const TCHAR* const Settings = TEXT("/Script/MacTargetPlatform.MacTargetSettings");
 #else
 	const TCHAR* const Settings = TEXT("/Script/IOSRuntimeSettings.IOSRuntimeSettings");
 #endif
-    if(!GConfig->GetInt(Settings, TEXT("MaxShaderLanguageVersion"), MaxShaderVersion, GEngineIni))
-	{
-		MaxShaderVersion = DefaultMaxShaderVersion;
-	}
-	if(!GConfig->GetInt(Settings, TEXT("IndirectArgumentTier"), IndirectArgumentTier, GEngineIni))
+    GConfig->GetInt(Settings, TEXT("MetalLanguageVersion"), MetalShaderVersion, GEngineIni);
+	
+    if(!GConfig->GetInt(Settings, TEXT("IndirectArgumentTier"), IndirectArgumentTier, GEngineIni))
 	{
 		IndirectArgumentTier = 0;
 	}
-	MaxShaderVersion = FMath::Max(MinShaderVersion, MaxShaderVersion);
-	ValidateVersion(MaxShaderVersion);
+	ValidateVersion(MetalShaderVersion);
 
 	if(MaxNumCommandBuffers == 0)
 	{
diff --git a/Engine/Source/Runtime/Apple/MetalRHI/Private/MetalShaders.cpp b/Engine/Source/Runtime/Apple/MetalRHI/Private/MetalShaders.cpp
--- a/Engine/Source/Runtime/Apple/MetalRHI/Private/MetalShaders.cpp
+++ b/Engine/Source/Runtime/Apple/MetalRHI/Private/MetalShaders.cpp
@@ -76,11 +76,14 @@ mtlpp::LanguageVersion ValidateVersion(uint32 Version)
 		case 5:
 			Result = mtlpp::LanguageVersion::Version2_2;
 			break;
-		default:
+        case 0:
             Version = 5;
-			Result = mtlpp::LanguageVersion::Version2_2;
-			break;
-	}
-	
+			Result = mtlpp::LanguageVersion::Version2_2; // minimum version as of UE5.0
+        default:
+            //EMacMetalShaderStandard::MacMetalSLStandard_Minimum and EIOSMetalShaderStandard::IOSMetalSLStandard_Minimum is currently 2.2
+            UE_LOG(LogTemp, Warning, TEXT("The Metal version currently set is not supported anymore. Set it in the Project Settings. Defaulting to the minimum version."));
+            Version = 5;
+            Result = mtlpp::LanguageVersion::Version2_2;
+            break;
 	return Result;
 }
diff --git a/Engine/Source/Runtime/Engine/Private/ShaderCompiler/ShaderCompiler.cpp b/Engine/Source/Runtime/Engine/Private/ShaderCompiler/ShaderCompiler.cpp
--- a/Engine/Source/Runtime/Engine/Private/ShaderCompiler/ShaderCompiler.cpp
+++ b/Engine/Source/Runtime/Engine/Private/ShaderCompiler/ShaderCompiler.cpp
@@ -5622,8 +5622,8 @@ void GlobalBeginCompileShader(
 		}
 		
 		{
-			uint32 ShaderVersion = RHIGetShaderLanguageVersion(EShaderPlatform(Target.Platform));
-			Input.Environment.SetDefine(TEXT("MAX_SHADER_LANGUAGE_VERSION"), ShaderVersion);
+			uint32 ShaderVersion = RHIGetMetalShaderLanguageVersion(EShaderPlatform(Target.Platform));
+			Input.Environment.SetDefine(TEXT("SHADER_LANGUAGE_VERSION"), ShaderVersion);
 			
 			bool bAllowFastIntrinsics = false;
 			bool bForceFloats = false;
diff --git a/Engine/Source/Runtime/IOS/IOSRuntimeSettings/Classes/IOSRuntimeSettings.h b/Engine/Source/Runtime/IOS/IOSRuntimeSettings/Classes/IOSRuntimeSettings.h
--- a/Engine/Source/Runtime/IOS/IOSRuntimeSettings/Classes/IOSRuntimeSettings.h
+++ b/Engine/Source/Runtime/IOS/IOSRuntimeSettings/Classes/IOSRuntimeSettings.h
@@ -44,6 +44,8 @@ UENUM()
 UENUM()
 enum class EIOSMetalShaderStandard : uint8
 {
+    /** Metal Shader 2.2 is the minimum as of UE5.0*/
+    IOSMetalSLStandard_Minimum = 0 UMETA(DisplayName="Minimum, Currently v2.2 (iOS 13.0/tvOS 13.0)"),
     /** Metal Shaders Compatible With iOS 13.0/tvOS 13.0 or later (std=ios-metal2.2) */
     IOSMetalSLStandard_2_2 = 5 UMETA(DisplayName="Metal v2.2 (iOS 13.0/tvOS 13.0)"),
     /** Metal Shaders Compatible With iOS 14.0/tvOS 14.0 or later (std=ios-metal2.3) */
@@ -439,10 +441,9 @@ public:
 	bool bDisableHTTPS;
 
 
-    // The maximum supported Metal shader langauge version.
-    // This defines what features may be used and OS versions supported.
-    UPROPERTY(EditAnywhere, config, Category=Rendering, meta = (DisplayName = "Max. Metal Shader Standard To Target", ConfigRestartRequired = true))
-	uint8 MaxShaderLanguageVersion;
+    // The Metal shader language version which will be used when compiling the shaders.
+    UPROPERTY(EditAnywhere, config, Category=Rendering, meta = (DisplayName = "Metal Shader Standard To Target", ConfigRestartRequired = true))
+	uint8 MetalLanguageVersion;
 	
 	/**
 	 * Whether to use the Metal shading language's "fast" intrinsics.
diff --git a/Engine/Source/Runtime/RHI/Private/RHI.cpp b/Engine/Source/Runtime/RHI/Private/RHI.cpp
--- a/Engine/Source/Runtime/RHI/Private/RHI.cpp
+++ b/Engine/Source/Runtime/RHI/Private/RHI.cpp
@@ -1773,44 +1773,36 @@ RHI_API const TCHAR* RHIVendorIdToString(EGpuVendorId VendorId)
 	return TEXT("Unknown");
 }
 
-RHI_API uint32 RHIGetShaderLanguageVersion(const FStaticShaderPlatform Platform)
+RHI_API uint32 RHIGetMetalShaderLanguageVersion(const FStaticShaderPlatform Platform)
 {
-	uint32 Version = 0;
-	if (IsMetalPlatform(Platform))
+    if (IsMetalPlatform(Platform))
 	{
 		if (IsPCPlatform(Platform))
 		{
-			static int32 MaxShaderVersion = -1;
-			if (MaxShaderVersion < 0)
-			{
-				MaxShaderVersion = 7; // MSL v2.4
-				int32 MinShaderVersion = 5; // MSL v2.2
-				if(!GConfig->GetInt(TEXT("/Script/MacTargetPlatform.MacTargetSettings"), TEXT("MaxShaderLanguageVersion"), MaxShaderVersion, GEngineIni))
-				{
-					MaxShaderVersion = 0;
-				}
-				MaxShaderVersion = FMath::Max(MinShaderVersion, MaxShaderVersion);
-			}
-			Version = (uint32)MaxShaderVersion;
+            static int32 MacMetalShaderLanguageVersion = -1;
+            if (MacMetalShaderLanguageVersion == -1)
+            {
+                if (!GConfig->GetInt(TEXT("/Script/MacTargetPlatform.MacTargetSettings"), TEXT("MetalLanguageVersion"), MacMetalShaderLanguageVersion, GEngineIni))
+                {
+                    MacMetalShaderLanguageVersion = 0; // 0 means default EMacMetalShaderStandard::MacMetalSLStandard_Minimum
+                }
+            }
+            return MacMetalShaderLanguageVersion;
 		}
 		else
 		{
-			static int32 MaxShaderVersion = -1;
-			if (MaxShaderVersion < 0)
-			{
-				MaxShaderVersion = 7;
-				int32 MinShaderVersion = 5;
-				if(!GConfig->GetInt(TEXT("/Script/IOSRuntimeSettings.IOSRuntimeSettings"), TEXT("MaxShaderLanguageVersion"), MaxShaderVersion, GEngineIni))
-				{
-					MaxShaderVersion = 0;
-				}
-                
-				MaxShaderVersion = FMath::Max(MinShaderVersion, MaxShaderVersion);
-			}
-			Version = (uint32)MaxShaderVersion;
-		}
+            static int32 IOSMetalShaderLanguageVersion = -1;
+            if (IOSMetalShaderLanguageVersion == -1)
+            {
+                if (!GConfig->GetInt(TEXT("/Script/IOSRuntimeSettings.IOSRuntimeSettings"), TEXT("MetalLanguageVersion"), IOSMetalShaderLanguageVersion, GEngineIni))
+                {
+                	IOSMetalShaderLanguageVersion = 0;  // 0 means default EIOSMetalShaderStandard::IOSMetalSLStandard_Minimum
+                }
+            }
+            return IOSMetalShaderLanguageVersion;
+        }
 	}
-	return Version;
+	return 0;
 }
 
 static ERHIFeatureLevel::Type GRHIMobilePreviewFeatureLevel = ERHIFeatureLevel::Num;
@@ -1854,7 +1846,7 @@ RHI_API int32 RHIGetPreferredClearUAVRectPSResourceType(const FStaticShaderPlatf
 	if (IsMetalPlatform(Platform))
 	{
 		static constexpr uint32 METAL_TEXTUREBUFFER_SHADER_LANGUAGE_VERSION = 4;
-		if (METAL_TEXTUREBUFFER_SHADER_LANGUAGE_VERSION <= RHIGetShaderLanguageVersion(Platform))
+		if (METAL_TEXTUREBUFFER_SHADER_LANGUAGE_VERSION <= RHIGetMetalShaderLanguageVersion(Platform))
 		{
 			return 0; // BUFFER
 		}
diff --git a/Engine/Source/Runtime/RHI/Public/RHI.h b/Engine/Source/Runtime/RHI/Public/RHI.h
--- a/Engine/Source/Runtime/RHI/Public/RHI.h
+++ b/Engine/Source/Runtime/RHI/Public/RHI.h
@@ -119,8 +119,8 @@ RHI_API const TCHAR* RHIVendorIdToString();
 // helper to convert VendorId into a printable string, or "Unknown" if unknown.
 RHI_API const TCHAR* RHIVendorIdToString(EGpuVendorId VendorId);
 
-// helper to return the shader language version for the given shader platform.
-RHI_API uint32 RHIGetShaderLanguageVersion(const FStaticShaderPlatform Platform);
+// helper to return the shader language version for Metal shader.
+RHI_API uint32 RHIGetMetalShaderLanguageVersion(const FStaticShaderPlatform Platform);
 
 // helper to check that the shader platform supports creating a UAV off an index buffer.
 inline bool RHISupportsIndexBufferUAVs(const FStaticShaderPlatform Platform)
diff --git a/Engine/Source/Runtime/RenderCore/Private/Shader.cpp b/Engine/Source/Runtime/RenderCore/Private/Shader.cpp
--- a/Engine/Source/Runtime/RenderCore/Private/Shader.cpp
+++ b/Engine/Source/Runtime/RenderCore/Private/Shader.cpp
@@ -1611,7 +1611,7 @@ void ShaderMapAppendKeyString(EShaderPlatform Platform, FString& KeyString)
 			KeyString += RHISupportsManualVertexFetch(Platform) ? TEXT("_MVF_") : TEXT("");
 		}
 		
-		uint32 ShaderVersion = RHIGetShaderLanguageVersion(Platform);
+		uint32 ShaderVersion = RHIGetMetalShaderLanguageVersion(Platform);
 		KeyString += FString::Printf(TEXT("_MTLSTD%u_"), ShaderVersion);
 		
 		bool bAllowFastIntrinsics = false;
diff --git a/Engine/Source/Runtime/Renderer/Private/PostProcess/PostProcessFFTBloom.cpp b/Engine/Source/Runtime/Renderer/Private/PostProcess/PostProcessFFTBloom.cpp
--- a/Engine/Source/Runtime/Renderer/Private/PostProcess/PostProcessFFTBloom.cpp
+++ b/Engine/Source/Runtime/Renderer/Private/PostProcess/PostProcessFFTBloom.cpp
@@ -206,18 +206,11 @@ float GetFFTBloomResolutionFraction(const FIntPoint& ViewSize)
 
 bool IsFFTBloomEnabled(const FViewInfo& View)
 {
-	const bool bOldMetalNoFFT = IsMetalPlatform(View.GetShaderPlatform()) && (RHIGetShaderLanguageVersion(View.GetShaderPlatform()) < 4) && IsPCPlatform(View.GetShaderPlatform());
 	const bool bUseFFTBloom = View.FinalPostProcessSettings.BloomMethod == EBloomMethod::BM_FFT && View.ViewState != nullptr && DoesPlatformSupportFFTBloom(View.GetShaderPlatform());
 
 	static bool bWarnAboutOldMetalFFTOnce = false;
 
-	if (bOldMetalNoFFT && bUseFFTBloom && !bWarnAboutOldMetalFFTOnce)
-	{
-		UE_LOG(LogRenderer, Error, TEXT("FFT Bloom is only supported on Metal 2.1 and later."));
-		bWarnAboutOldMetalFFTOnce = true;
-	}
-
-	if (bUseFFTBloom && !bOldMetalNoFFT)
+	if (bUseFFTBloom)
 	{
 		return View.FFTBloomKernelTexture != nullptr;
 	}
