Fix enum warning in ios settings

diff --git a/Engine/Source/Developer/IOS/IOSPlatformEditor/Private/IOSTargetSettingsCustomization.cpp b/Engine/Source/Developer/IOS/IOSPlatformEditor/Private/IOSTargetSettingsCustomization.cpp
--- a/Engine/Source/Developer/IOS/IOSPlatformEditor/Private/IOSTargetSettingsCustomization.cpp
+++ b/Engine/Source/Developer/IOS/IOSPlatformEditor/Private/IOSTargetSettingsCustomization.cpp
@@ -1954,7 +1954,7 @@ void FIOSTargetSettingsCustomization::UpdateMetalMRTWarning()
 			}
 			
 			ShaderVersionPropertyHandle->GetValue(EnumValue);
-			if (EnumValue < (int32)EIOSMetalShaderStandard::IOSMetalSLStandard_Minimum)
+			if (EnumValue < (int32)EIOSMetalShaderStandard::IOSMetalSLStandard_2_2)
 			{
 				SetShaderStandard((int32)EIOSMetalShaderStandard::IOSMetalSLStandard_Minimum);
 			}
