Added new flag to enable Apple A8 support for mobile devices
Disabled texture compression for 3D textures

diff --git a/Engine/Source/Developer/IOS/IOSTargetPlatform/Private/IOSTargetPlatform.cpp b/Engine/Source/Developer/IOS/IOSTargetPlatform/Private/IOSTargetPlatform.cpp
--- a/Engine/Source/Developer/IOS/IOSTargetPlatform/Private/IOSTargetPlatform.cpp
+++ b/Engine/Source/Developer/IOS/IOSTargetPlatform/Private/IOSTargetPlatform.cpp
@@ -441,6 +441,14 @@ static bool SupportsMetalMRT()
 	return bSupportsMetalMRT;
 }
 
+static bool SupportsA8Devices()
+{
+    // default to NOT supporting A8 devices
+    bool bSupportAppleA8 = false;
+    GConfig->GetBool(TEXT("/Script/IOSRuntimeSettings.IOSRuntimeSettings"), TEXT("bSupportAppleA8"), bSupportAppleA8, GEngineIni);
+    return bSupportAppleA8;
+}
+
 static bool SupportsLandscapeMeshLODStreaming()
 {
 	bool bStreamLandscapeMeshLODs = false;
@@ -583,7 +591,7 @@ void FIOSTargetPlatform::GetTextureFormats( const UTexture* Texture, TArray< TAr
         int32 BlockSize = 1;
 		// Compressed volume textures require MTLGPUFamilyApple3 or later
 		// min spec for TVOS is AppleTV HD which is MTLGPUFamilyApple2 (A8)
-		bool bSupportCompressedVolumeTexture = !bIsTVOS;
+		bool bSupportCompressedVolumeTexture = !bIsTVOS && !SupportsA8Devices();
 		GetDefaultTextureFormatNamePerLayer(TextureFormatNames, this, Texture, true, bSupportCompressedVolumeTexture, BlockSize);
 	}
 
diff --git a/Engine/Source/Runtime/IOS/IOSRuntimeSettings/Classes/IOSRuntimeSettings.h b/Engine/Source/Runtime/IOS/IOSRuntimeSettings/Classes/IOSRuntimeSettings.h
--- a/Engine/Source/Runtime/IOS/IOSRuntimeSettings/Classes/IOSRuntimeSettings.h
+++ b/Engine/Source/Runtime/IOS/IOSRuntimeSettings/Classes/IOSRuntimeSettings.h
@@ -470,6 +470,10 @@ public:
 	UPROPERTY(config, EditAnywhere, Category = Rendering, Meta = (DisplayName = "Tier of Indirect Argument Buffers to use when compiling shaders", ConfigRestartRequired = true))
 	int32 IndirectArgumentTier;
 	
+    /** Supports Apple A8 devices, disables 3d texture compression */
+    UPROPERTY(config, EditAnywhere, Category = Rendering, Meta = (DisplayName = "Support Apple A8", ConfigRestartRequired = true))
+    bool bSupportAppleA8;
+    
 	// Whether or not the keyboard should be usable on it's own without a UITextField
 	UPROPERTY(GlobalConfig, EditAnywhere, Category = Input)
 	bool bUseIntegratedKeyboard;
