Refactor TVOS Asset catalog : optional and required assets, assets generation, 4K Top Shelf support

diff --git a/Engine/Build/TVOS/Resources/Assets.xcassets/App Icon & Top Shelf Image.brandassets/Contents.json b/Engine/Build/TVOS/Resources/Assets.xcassets/App Icon & Top Shelf Image.brandassets/Contents.json
--- a/Engine/Build/TVOS/Resources/Assets.xcassets/App Icon & Top Shelf Image.brandassets/Contents.json	
+++ b/Engine/Build/TVOS/Resources/Assets.xcassets/App Icon & Top Shelf Image.brandassets/Contents.json	
@@ -17,12 +17,6 @@
       "idiom" : "tv",
       "filename" : "Top Shelf Image Wide.imageset",
       "role" : "top-shelf-image-wide"
-    },
-    {
-      "size" : "1920x720",
-      "idiom" : "tv",
-      "filename" : "Top Shelf Image.imageset",
-      "role" : "top-shelf-image"
     }
   ],
   "info" : {
diff --git a/Engine/Build/TVOS/Resources/Assets.xcassets/App Icon & Top Shelf Image.brandassets/Top Shelf Image.imageset/Contents.json b/Engine/Build/TVOS/Resources/Assets.xcassets/App Icon & Top Shelf Image.brandassets/Top Shelf Image.imageset/Contents.json
deleted file mode 100644
diff --git a/Engine/Source/Developer/IOS/IOSPlatformEditor/Private/IOSTargetSettingsCustomization.cpp b/Engine/Source/Developer/IOS/IOSPlatformEditor/Private/IOSTargetSettingsCustomization.cpp
--- a/Engine/Source/Developer/IOS/IOSPlatformEditor/Private/IOSTargetSettingsCustomization.cpp
+++ b/Engine/Source/Developer/IOS/IOSPlatformEditor/Private/IOSTargetSettingsCustomization.cpp
@@ -72,7 +72,7 @@ FIOSTargetSettingsCustomization::FIOSTargetSettingsCustomization()
 
     // Default AppIcons copied at the payload's root. See https://developer.apple.com/design/human-interface-guidelines/ios/icons-and-images/app-icon/
    
-	new (IconNames)FPlatformIconInfo(TEXT("Icon1024.png"), LOCTEXT("AppIcon_Marketing", "Marketing Icon (1024x1024)\n\nOther icons sizes can be generated from the Marketing Icon."), FText::GetEmpty(), 1024, 1024, FPlatformIconInfo::Required); // App Store
+	new (IconNames)FPlatformIconInfo(TEXT("Icon1024.png"), LOCTEXT("AppIcon_Marketing", "Marketing Icon (1024x1024)\n\nOther iOS icons sizes can be generated from the Marketing Icon."), FText::GetEmpty(), 1024, 1024, FPlatformIconInfo::Required); // App Store
 
 	new (IconNames) FPlatformIconInfo(TEXT("[REDACTED_EMAIL]"), LOCTEXT("Default_iPhone_AppIcon", "Default iPhone Icon (120x120)"), FText::GetEmpty(), 120, 120, FPlatformIconInfo::Required); // iPhone
     new (IconNames) FPlatformIconInfo(TEXT("[REDACTED_EMAIL]"), LOCTEXT("Default_iPad_AppIcon", "Default iPad App Icon (152x152)"), FText::GetEmpty(), 152, 152, FPlatformIconInfo::Required); // iPad, iPad Mini
@@ -98,16 +98,27 @@ FIOSTargetSettingsCustomization::FIOSTargetSettingsCustomization()
 	new (LaunchImageNames)FPlatformIconInfo(TEXT("LaunchScreenIOS.png"), LOCTEXT("LaunchImageIOS", "Launch Screen Image"), LOCTEXT("LaunchImageIOSDesc", "This image is used for the Launch Screen when custom launch screen storyboards are not in use. The image is used in both portait and landscape modes and will be uniformly scaled to occupy the full width or height as necessary for of all devices, so if your app supports both a square image is recommended. The png file supplied must not have an alpha channel."), -1, -1, FPlatformIconInfo::Required);
 
 	// Icons and Shelf Images for tvOS
-	new (TvOSImageNames)FPlatformIconInfo(TEXT("Icon_Large_Back.png"), LOCTEXT("TVOS_Icon_Large_Back", "Icon Large Back (1280x768)"), FText::GetEmpty(), 1280, 768, FPlatformIconInfo::Required);
-	new (TvOSImageNames)FPlatformIconInfo(TEXT("Icon_Large_Front.png"), LOCTEXT("TVOS_Icon_Large_Front", "Icon Large Front (1280x768)"), FText::GetEmpty(), 1280, 768, FPlatformIconInfo::Required);
-	new (TvOSImageNames)FPlatformIconInfo(TEXT("Icon_Large_Middle.png"), LOCTEXT("TVOS_Icon_Large_Middle", "Icon Large Middle (1280x768)"), FText::GetEmpty(), 1280, 768, FPlatformIconInfo::Required);
-	new (TvOSImageNames)FPlatformIconInfo(TEXT("Icon_Small_Back.png"), LOCTEXT("TVOS_Icon_Small_Back", "Icon Small Back (400x240)"), FText::GetEmpty(), 400, 240, FPlatformIconInfo::Required);
-	new (TvOSImageNames)FPlatformIconInfo(TEXT("Icon_Small_Front.png"), LOCTEXT("TVOS_Icon_Small_Front", "Icon Small Front (400x240)"), FText::GetEmpty(), 400, 240, FPlatformIconInfo::Required);
-	new (TvOSImageNames)FPlatformIconInfo(TEXT("Icon_Small_Middle.png"), LOCTEXT("TVOS_Icon_Small_Middle", "Icon Small Middle (400x240)"), FText::GetEmpty(), 400, 240, FPlatformIconInfo::Required);
-	new (TvOSImageNames)FPlatformIconInfo(TEXT("TopShelf.png"), LOCTEXT("TVOS_Top_Shelf", "Top Shelf (1920x720)"), FText::GetEmpty(), 1920, 720, FPlatformIconInfo::Required);
-	new (TvOSImageNames)FPlatformIconInfo(TEXT("[REDACTED_EMAIL]"), LOCTEXT("2x_TVOS_Top_Shelf", "Top Shelf (6840x1440"), FText::GetEmpty(), 3840, 1440, FPlatformIconInfo::Required);
-	new (TvOSImageNames)FPlatformIconInfo(TEXT("TopShelfWide-1920x720.png"), LOCTEXT("TVOS_Top_Shelf_Wide", "Top Shelf Wide (2320x720)"), FText::GetEmpty(), 2320, 720, FPlatformIconInfo::Required);
-	new (TvOSImageNames)FPlatformIconInfo(TEXT("[REDACTED_EMAIL]"), LOCTEXT("2x_TVOS_Top_Shelf_Wide", "2x Top Shelf Wide (4640x1440)"), FText::GetEmpty(), 4640, 1440, FPlatformIconInfo::Required);
+
+    // Used to generate top shelf images
+    new (TvOSImageNames)FPlatformIconInfo(TEXT("[REDACTED_EMAIL]"), LOCTEXT("2x_TVOS_Top_Shelf_Wide", "2x Top Shelf Wide (4640x1440)\n\nOther tvOS Topshelf Image sizes can be generated from it."), FText::GetEmpty(), 4640, 1440, FPlatformIconInfo::Required);
+
+    // Generated top shelf image
+    new (TvOSImageNames)FPlatformIconInfo(TEXT("TopShelfWide-2320x720.png"), LOCTEXT("TVOS_Top_Shelf_Wide", "Top Shelf Wide (2320x720)"), FText::GetEmpty(), 2320, 720, FPlatformIconInfo::Optional);
+
+    // Used to generate other tvOS icons
+    new (TvOSImageNames)FPlatformIconInfo(TEXT("Icon_Large_Front.png"), LOCTEXT("TVOS_Icon_Large_Front", "Icon Large Front (1280x768)\n\nOther tvOS icons sizes can be generated from it."), FText::GetEmpty(), 1280, 768, FPlatformIconInfo::Required);
+
+    // Generated icons
+    new (TvOSImageNames)FPlatformIconInfo(TEXT("Icon_Large_Middle.png"), LOCTEXT("TVOS_Icon_Large_Middle", "Icon Large Middle (1280x768)"), FText::GetEmpty(), 1280, 768, FPlatformIconInfo::Required);
+    new (TvOSImageNames)FPlatformIconInfo(TEXT("Icon_Large_Back.png"), LOCTEXT("TVOS_Icon_Large_Back", "Icon Large Back (1280x768)"), FText::GetEmpty(), 1280, 768, FPlatformIconInfo::Required);
+    
+    new (TvOSImageNames)FPlatformIconInfo(TEXT("Icon_Medium_Front.png"), LOCTEXT("TVOS_Icon_Medium_Front", "Icon Medium Front (800x480)"), FText::GetEmpty(), 800, 480, FPlatformIconInfo::Required);
+    new (TvOSImageNames)FPlatformIconInfo(TEXT("Icon_Medium_Middle.png"), LOCTEXT("TVOS_Icon_Medium_Middle", "Icon Medium Middle (800x480)"), FText::GetEmpty(), 800, 480, FPlatformIconInfo::Required);
+    new (TvOSImageNames)FPlatformIconInfo(TEXT("Icon_Medium_Back.png"), LOCTEXT("TVOS_Icon_Medium_Back", "Icon Medium Back (800x480)"), FText::GetEmpty(), 800, 480, FPlatformIconInfo::Required);
+
+    new (TvOSImageNames)FPlatformIconInfo(TEXT("Icon_Small_Front.png"), LOCTEXT("TVOS_Icon_Small_Front", "Icon Small Front (400x240)"), FText::GetEmpty(), 400, 240, FPlatformIconInfo::Optional);
+    new (TvOSImageNames)FPlatformIconInfo(TEXT("Icon_Small_Middle.png"), LOCTEXT("TVOS_Icon_Small_Middle", "Icon Small Middle (400x240)"), FText::GetEmpty(), 400, 240, FPlatformIconInfo::Optional);
+    new (TvOSImageNames)FPlatformIconInfo(TEXT("Icon_Small_Back.png"), LOCTEXT("TVOS_Icon_Small_Back", "Icon Small Back (400x240)"), FText::GetEmpty(), 400, 240, FPlatformIconInfo::Optional);
 
 	bShowAllProvisions = false;
 	bShowAllCertificates = false;
@@ -1056,34 +1067,38 @@ void FIOSTargetSettingsCustomization::BuildRemoteBuildingSection(IDetailLayoutBu
 
 void FIOSTargetSettingsCustomization::BuildIconSection(IDetailLayoutBuilder& DetailLayout)
 {
-	IDetailCategoryBuilder& RequiredIconCategory = DetailLayout.EditCategory(TEXT("Required Icons"));
-	IDetailCategoryBuilder& OptionalIconCategory = DetailLayout.EditCategory(TEXT("Optional Icons"));
+	IDetailCategoryBuilder& RequiredIconCategoryIOS = DetailLayout.EditCategory("RequiredIOSIcons", LOCTEXT("RequiredIOSIcons", "Required Icons (iPhone and iPad)"));
+    IDetailCategoryBuilder& OptionalIconCategoryIOS = DetailLayout.EditCategory("OptionalIOSIcons", LOCTEXT("OptionalIOSIcons", "Optional Icons (iPhone and iPad)"));
 
-	// Add the icons
+	// Add the iOS icons
 	for (const FPlatformIconInfo& Info : IconNames)
 	{
 		FVector2D IconImageMaxSize(Info.IconRequiredSize);
 		IconImageMaxSize.X = FMath::Min(IconImageMaxSize.X, 150.0f);
 		IconImageMaxSize.Y = FMath::Min(IconImageMaxSize.Y, 150.0f);
-		IDetailCategoryBuilder& IconCategory = (Info.RequiredState == FPlatformIconInfo::Required) ? RequiredIconCategory : OptionalIconCategory;
+		IDetailCategoryBuilder& IconCategory = (Info.RequiredState == FPlatformIconInfo::Required) ? RequiredIconCategoryIOS : OptionalIconCategoryIOS;
 		BuildImageRow(DetailLayout, IconCategory, Info, IconImageMaxSize);
 	}
 
-	// Add the launch images
-	IDetailCategoryBuilder& LaunchImageCategory = DetailLayout.EditCategory(FName("LaunchScreen"));
-	const FVector2D LaunchImageMaxSize(150.0f, 150.0f);
-	for (const FPlatformIconInfo& Info : LaunchImageNames)
-	{
-		BuildImageRow(DetailLayout, LaunchImageCategory, Info, LaunchImageMaxSize);
-	}
-
 	// Add the tvOS content
-	IDetailCategoryBuilder& tvOSCategory = DetailLayout.EditCategory(FName("tvOS"));
+    IDetailCategoryBuilder& RequiredIconCategoryTvOS = DetailLayout.EditCategory("RequiredTVOSAssets", LOCTEXT("RequiredTVOSAssets", "Required Assets (AppleTV)"));
+    IDetailCategoryBuilder& OptionalIconCategoryTvOS = DetailLayout.EditCategory("OptionalTVOSAssets", LOCTEXT("OptionalTVOSAssets", "Optional Assets (AppleTV)"));
+
 	const FVector2D TvOSImageMaxSize(150.0f, 150.0f);
 	for (const FPlatformIconInfo& Info : TvOSImageNames)
 	{
-		BuildImageRow(DetailLayout, tvOSCategory, Info, TvOSImageMaxSize, true);
+        IDetailCategoryBuilder& ImageCategory = (Info.RequiredState == FPlatformIconInfo::Required) ? RequiredIconCategoryTvOS : OptionalIconCategoryTvOS;
+		BuildImageRow(DetailLayout, ImageCategory, Info, TvOSImageMaxSize, true);
 	}
+    
+    
+    // Add the launch images
+    IDetailCategoryBuilder& LaunchImageCategory = DetailLayout.EditCategory(FName("LaunchScreen"));
+    const FVector2D LaunchImageMaxSize(150.0f, 150.0f);
+    for (const FPlatformIconInfo& Info : LaunchImageNames)
+    {
+        BuildImageRow(DetailLayout, LaunchImageCategory, Info, LaunchImageMaxSize);
+    }
 
 }
 
@@ -1131,13 +1146,20 @@ void FIOSTargetSettingsCustomization::CopySetupFilesIntoProject()
 
 void FIOSTargetSettingsCustomization::BuildImageRow(IDetailLayoutBuilder& DetailLayout, IDetailCategoryBuilder& Category, const FPlatformIconInfo& Info, const FVector2D& MaxDisplaySize, bool bIsTVOS)
 {
-	FString AutomaticImagePath = EngineGraphicsPath / Info.IconPath;
+    FString AutomaticImagePath = EngineGraphicsPath / Info.IconPath;
 	FString TargetImagePath = GameGraphicsPath / Info.IconPath;
+    FString SourceImagePath = FPaths::GetPath(FPaths::GetProjectFilePath()) + TEXT("/Build/IOS/Resources/Graphics/Icon1024.png");
 
 	if (bIsTVOS)
 	{
 		AutomaticImagePath = TVOSEngineGraphicsPath / Info.IconPath;
 		TargetImagePath = TVOSGameGraphicsPath / Info.IconPath;
+        SourceImagePath = FPaths::GetPath(FPaths::GetProjectFilePath()) + TEXT("/Build/TVOS/Resources/Graphics/Icon_Large_Front.png");
+
+        if (Info.IconName.ToString().Contains("Top Shelf"))
+        {
+             SourceImagePath = FPaths::GetPath(FPaths::GetProjectFilePath()) + TEXT("[REDACTED_EMAIL]");
+        }
 	}
 
     if (Info.RequiredState == FPlatformIconInfo::Required)
@@ -1169,9 +1191,10 @@ void FIOSTargetSettingsCustomization::BuildImageRow(IDetailLayoutBuilder& Detail
                     SNew(SExternalImageReference, AutomaticImagePath, TargetImagePath)
                     .RequiredSize(Info.IconRequiredSize)
                     .MaxDisplaySize(MaxDisplaySize)
-					.GenerateImageVisibility(this, &FIOSTargetSettingsCustomization::ShouldShowGenerateButtonForIcon, bIsTVOS || Info.IconName.ToString().Contains("Launch Screen Image") || Info.IconName.ToString().Contains("Marketing Icon"))
-					.GenerateImageToolTipText(LOCTEXT("GenerateFromMarketingIcon", "Generate from Marketing Icon"))
-					.OnGenerateImageClicked(this, &FIOSTargetSettingsCustomization::OnGenerateImageClicked, TargetImagePath, Info.IconRequiredSize)
+					.GenerateImageVisibility(this, &FIOSTargetSettingsCustomization::ShouldShowGenerateButtonForIcon, Info.IconName.ToString().Contains("Launch Screen Image") || Info.IconName.ToString().Contains("Marketing Icon") ||
+                                             Info.IconName.ToString().Contains("Icon Large Front (1280x768)") || Info.IconName.ToString().Contains("2x Top Shelf Wide (4640x1440)"), SourceImagePath)
+					.GenerateImageToolTipText(LOCTEXT("GenerateFromOtherIcon", "Generate from Bigger Image (see image tooltip)"))
+					.OnGenerateImageClicked(this, &FIOSTargetSettingsCustomization::OnGenerateImageClicked, SourceImagePath, TargetImagePath, Info.IconRequiredSize)
                     .DeleteTargetWhenDefaultChosen(true)
                 ]
             ];
@@ -1205,9 +1228,10 @@ void FIOSTargetSettingsCustomization::BuildImageRow(IDetailLayoutBuilder& Detail
                     SNew(SExternalImageReference, "", TargetImagePath)
                     .RequiredSize(Info.IconRequiredSize)
                     .MaxDisplaySize(MaxDisplaySize)
-					.GenerateImageVisibility(this, &FIOSTargetSettingsCustomization::ShouldShowGenerateButtonForIcon, bIsTVOS || Info.IconName.ToString().Contains("Launch Screen Image") || Info.IconName.ToString().Contains("Marketing Icon"))
-					.GenerateImageToolTipText(LOCTEXT("GenerateFromMarketingIcon", "Generate from Marketing Icon"))
-					.OnGenerateImageClicked(this, &FIOSTargetSettingsCustomization::OnGenerateImageClicked, TargetImagePath, Info.IconRequiredSize)
+					.GenerateImageVisibility(this, &FIOSTargetSettingsCustomization::ShouldShowGenerateButtonForIcon, Info.IconName.ToString().Contains("Launch Screen Image") || Info.IconName.ToString().Contains("Marketing Icon") ||
+                                             Info.IconName.ToString().Contains("Icon Large Front (1280x768)") || Info.IconName.ToString().Contains("2x Top Shelf Wide (4640x1440)"), SourceImagePath)
+					.GenerateImageToolTipText(LOCTEXT("GenerateFromOtherIcon", "Generate from Bigger Image (see image tooltip)"))
+					.OnGenerateImageClicked(this, &FIOSTargetSettingsCustomization::OnGenerateImageClicked, SourceImagePath, TargetImagePath, Info.IconRequiredSize)
 				]
             ];
     }
@@ -1322,10 +1346,10 @@ FReply FIOSTargetSettingsCustomization::OnInstallProvisionClicked()
 	return FReply::Handled();
 }
 
-FReply FIOSTargetSettingsCustomization::OnGenerateImageClicked(const FString TargetImagePath, FIntPoint IconRequiredSize)
+FReply FIOSTargetSettingsCustomization::OnGenerateImageClicked(const FString SourceImagePath, const FString TargetImagePath, FIntPoint IconRequiredSize)
 {
-	FString GenerateImageSourcePath = FPaths::GetPath(FPaths::GetProjectFilePath()) + TEXT("/Build/IOS/Resources/Graphics/Icon1024.png");
-
+    UE_LOG(LogIOSTargetSettings, Warning, TEXT("HEY SALUT=[%s]"), *SourceImagePath);
+    
 	if (FPaths::FileExists(*TargetImagePath))
 	{
 		const EAppReturnType::Type Answer = FMessageDialog::Open(EAppMsgType::YesNo, FText::FromString(TEXT("File already exists. Do you want to overwrite it ?")));
@@ -1335,7 +1359,7 @@ FReply FIOSTargetSettingsCustomization::OnGenerateImageClicked(const FString Tar
 		}
 	}
 
-	if (ensure(FPaths::FileExists(*GenerateImageSourcePath)))
+	if (ensure(FPaths::FileExists(*SourceImagePath)))
 	{
 		FUEFreeImageWrapper::FreeImage_Initialise();
 		if (!FUEFreeImageWrapper::IsValid())
@@ -1345,39 +1369,39 @@ FReply FIOSTargetSettingsCustomization::OnGenerateImageClicked(const FString Tar
 		}
 
 		FREE_IMAGE_FORMAT FileType = FIF_UNKNOWN;
-		FileType = FreeImage_GetFileType(TCHAR_TO_FICHAR(*GenerateImageSourcePath), 0);
+		FileType = FreeImage_GetFileType(TCHAR_TO_FICHAR(*SourceImagePath), 0);
 		if (FileType == FIF_UNKNOWN)
 		{
-			FileType = FreeImage_GetFIFFromFilename(TCHAR_TO_FICHAR(*GenerateImageSourcePath));
+			FileType = FreeImage_GetFIFFromFilename(TCHAR_TO_FICHAR(*SourceImagePath));
 			if (FileType == FIF_UNKNOWN)
 			{
-				FMessageDialog::Open(EAppMsgType::Ok, LOCTEXT("Generate_UnknownFileType", "An unknown filetype error occurred while trying to resize the Marketing icon."));
+				FMessageDialog::Open(EAppMsgType::Ok, LOCTEXT("Generate_UnknownFileType", "An unknown filetype error occurred while trying to resize the image."));
 				return FReply::Unhandled();
 			}
 		}
 
-		FIBITMAP* Bitmap = FreeImage_Load(FileType, TCHAR_TO_FICHAR(*GenerateImageSourcePath), 0);
+		FIBITMAP* Bitmap = FreeImage_Load(FileType, TCHAR_TO_FICHAR(*SourceImagePath), 0);
 		if (Bitmap == nullptr)
 		{
-			FMessageDialog::Open(EAppMsgType::Ok, LOCTEXT("Generate_LoadFailed", "The Marketing icon file could not be loaded while trying to resize it."));
+			FMessageDialog::Open(EAppMsgType::Ok, LOCTEXT("Generate_LoadFailed", "The image file could not be loaded while trying to resize it."));
 			return FReply::Unhandled();
 		}
 
 		FIBITMAP* RescaledImage;
 		FREE_IMAGE_FORMAT FifW;
-		if ((RescaledImage = FreeImage_Rescale(Bitmap, IconRequiredSize.X, IconRequiredSize.Y, FREE_IMAGE_FILTER::FILTER_LANCZOS3)) == nullptr ||
-			(FifW = FreeImage_GetFIFFromFilename(TCHAR_TO_FICHAR(*TargetImagePath))) == FIF_UNKNOWN ||
-			!FreeImage_Save(FifW, RescaledImage, TCHAR_TO_FICHAR(*TargetImagePath), 0))
-		{
-			FMessageDialog::Open(EAppMsgType::Ok, LOCTEXT("Generate_ResizeSaveFailed", "An error occurred while resizing or saving the icon file."));
-			return FReply::Unhandled();
-		}
-
+        if ((RescaledImage = FreeImage_Rescale(Bitmap, IconRequiredSize.X, IconRequiredSize.Y, FREE_IMAGE_FILTER::FILTER_LANCZOS3)) == nullptr ||
+                (FifW = FreeImage_GetFIFFromFilename(TCHAR_TO_FICHAR(*TargetImagePath))) == FIF_UNKNOWN ||
+                !FreeImage_Save(FifW, RescaledImage, TCHAR_TO_FICHAR(*TargetImagePath), 0))
+            {
+                FMessageDialog::Open(EAppMsgType::Ok, LOCTEXT("Generate_ResizeSaveFailed", "An error occurred while resizing or saving the icon file."));
+                return FReply::Unhandled();
+            }
+        
 		return FReply::Handled();
 	}
 	else
 	{
-		FMessageDialog::Open(EAppMsgType::Ok, LOCTEXT("Generate_OpenFailed", "The Marketing icon file could not be found."));
+		FMessageDialog::Open(EAppMsgType::Ok, LOCTEXT("Generate_OpenFailed", "The image file could not be found."));
 		return FReply::Unhandled();
 	}
 }
@@ -1975,9 +1999,9 @@ void FIOSTargetSettingsCustomization::SetMinVersion(int32 Value)
 	check(Res == FPropertyAccess::Success);
 }
 
-EVisibility FIOSTargetSettingsCustomization::ShouldShowGenerateButtonForIcon(bool bCannotBeGenerated) const
+EVisibility FIOSTargetSettingsCustomization::ShouldShowGenerateButtonForIcon(bool bCannotBeGenerated, const FString ImageToCheck) const
 {
-	if (!bCannotBeGenerated && FPlatformFileManager::Get().GetPlatformFile().FileExists(*(FPaths::GetPath(FPaths::GetProjectFilePath()) + "/Build/IOS/Resources/Graphics/Icon1024.png")))
+	if (!bCannotBeGenerated && FPlatformFileManager::Get().GetPlatformFile().FileExists(*ImageToCheck))
 	{
 		return EVisibility::Visible;
 	}
diff --git a/Engine/Source/Developer/IOS/IOSPlatformEditor/Private/IOSTargetSettingsCustomization.h b/Engine/Source/Developer/IOS/IOSPlatformEditor/Private/IOSTargetSettingsCustomization.h
--- a/Engine/Source/Developer/IOS/IOSPlatformEditor/Private/IOSTargetSettingsCustomization.h
+++ b/Engine/Source/Developer/IOS/IOSPlatformEditor/Private/IOSTargetSettingsCustomization.h
@@ -63,8 +63,8 @@ public:
 	virtual void CustomizeDetails(IDetailLayoutBuilder& DetailLayout) override;
 	// End of IDetailCustomization interface
     
-    EVisibility ShouldShowGenerateButtonForIcon(bool bCannotBeGenerated) const;
-
+    EVisibility ShouldShowGenerateButtonForIcon(bool bCannotBeGenerated, FString ImageToCheck) const;
+    
 private:
 	TArray<struct FPlatformIconInfo> IconNames;
 	TArray<struct FPlatformIconInfo> LaunchImageNames;
@@ -137,7 +137,7 @@ private:
 	// certificate request
 	FReply OnCertificateRequestClicked();
 
-	FReply OnGenerateImageClicked(const FString TargetImagePath, FIntPoint IconRequiredSize);
+	FReply OnGenerateImageClicked(const FString SourceImagePath, const FString TargetImagePath, FIntPoint IconRequiredSize);
 
 	// ssh key request
 	FReply OnGenerateSSHKey();
diff --git a/Engine/Source/Programs/UnrealBuildTool/Platform/IOS/IOSToolChain.cs b/Engine/Source/Programs/UnrealBuildTool/Platform/IOS/IOSToolChain.cs
--- a/Engine/Source/Programs/UnrealBuildTool/Platform/IOS/IOSToolChain.cs
+++ b/Engine/Source/Programs/UnrealBuildTool/Platform/IOS/IOSToolChain.cs
@@ -1321,10 +1321,8 @@ namespace UnrealBuildTool
 					new string []{ "Icon_Small_Front-800x480.png", "App Icon & Top Shelf Image.brandassets/App Icon.imagestack/Front.imagestacklayer/Content.imageset" },
 					new string []{ "Icon_Small_Middle-400x240.png", "App Icon & Top Shelf Image.brandassets/App Icon.imagestack/Middle.imagestacklayer/Content.imageset" },
 					new string []{ "Icon_Small_Middle-800x480.png", "App Icon & Top Shelf Image.brandassets/App Icon.imagestack/Middle.imagestacklayer/Content.imageset" },
-					new string []{ "[REDACTED_EMAIL]", "App Icon & Top Shelf Image.brandassets/Top Shelf Image Wide.imageset" },
-					new string []{ "TopShelfWide-1920x720.png", "App Icon & Top Shelf Image.brandassets/Top Shelf Image Wide.imageset" },
-					new string []{ "TopShelf.png", "App Icon & Top Shelf Image.brandassets/Top Shelf Image.imageset" },
-					new string []{ "[REDACTED_EMAIL]", "App Icon & Top Shelf Image.brandassets/Top Shelf Image.imageset" },
+					new string []{ "[REDACTED_EMAIL]", "App Icon & Top Shelf Image.brandassets/Top Shelf Image Wide.imageset" },
+					new string []{ "TopShelfWide-2320x720.png", "App Icon & Top Shelf Image.brandassets/Top Shelf Image Wide.imageset" },
 				};
 				Dir = Path.Combine(IntermediateDir, "Resources", "Assets.xcassets");
 
diff --git a/Engine/Source/Programs/UnrealBuildTool/Platform/TVOS/UEDeployTVOS.cs b/Engine/Source/Programs/UnrealBuildTool/Platform/TVOS/UEDeployTVOS.cs
--- a/Engine/Source/Programs/UnrealBuildTool/Platform/TVOS/UEDeployTVOS.cs
+++ b/Engine/Source/Programs/UnrealBuildTool/Platform/TVOS/UEDeployTVOS.cs
@@ -132,8 +132,6 @@ namespace UnrealBuildTool
 
             Text.AppendLine("\t<key>TVTopShelfImage</key>");
             Text.AppendLine("\t<dict>");
-            Text.AppendLine("\t\t<key>TVTopShelfPrimaryImage</key>");
-            Text.AppendLine("\t\t<string>Top Shelf Image</string>");
 			Text.AppendLine("\t\t<key>TVTopShelfPrimaryImageWide</key>");
 			Text.AppendLine("\t\t<string>Top Shelf Image Wide</string>");
 			Text.AppendLine("\t</dict>");
