Remove UE4 reference from comment

diff --git a/Engine/Source/Developer/IOS/IOSPlatformEditor/Private/IOSTargetSettingsCustomization.cpp b/Engine/Source/Developer/IOS/IOSPlatformEditor/Private/IOSTargetSettingsCustomization.cpp
--- a/Engine/Source/Developer/IOS/IOSPlatformEditor/Private/IOSTargetSettingsCustomization.cpp
+++ b/Engine/Source/Developer/IOS/IOSPlatformEditor/Private/IOSTargetSettingsCustomization.cpp
@@ -1504,7 +1504,7 @@ FReply FIOSTargetSettingsCustomization::OnGenerateSSHKey()
 	FString CwRsyncPath = Settings.CwRsyncInstallPath.Path;
 	if (CwRsyncPath.IsEmpty() || !FPaths::DirectoryExists(CwRsyncPath))
 	{
-		// If no user specified directory try the UE4 bundled directory
+		// If no user specified directory try the bundled ThirdPartyNotUE directory
 		CwRsyncPath = FPaths::ConvertRelativePathToFull(FPaths::EngineDir() / TEXT("Extras\\ThirdPartyNotUE\\cwrsync\\bin"));
 	}
 
