Turn off device detection on iOS when libimobile util binary not found.

diff --git a/Engine/Source/Developer/IOS/IOSTargetPlatform/Private/IOSDeviceHelper.cpp b/Engine/Source/Developer/IOS/IOSTargetPlatform/Private/IOSDeviceHelper.cpp
--- a/Engine/Source/Developer/IOS/IOSTargetPlatform/Private/IOSDeviceHelper.cpp
+++ b/Engine/Source/Developer/IOS/IOSTargetPlatform/Private/IOSDeviceHelper.cpp
@@ -305,8 +305,15 @@ private:
     {
         FString OutStdOut;
         FString OutStdErr;
-        FString LibimobileDeviceId = GetLibImobileDeviceExe("idevice_id");
-        int ReturnCode;
+		FString LibimobileDeviceId = GetLibImobileDeviceExe("idevice_id");
+		int ReturnCode;
+		if (LibimobileDeviceId.Len() == 0)
+		{
+			UE_LOG(LogIOSDeviceHelper, Log, TEXT("idevice_id (iOS device detection) executable missing. Turning off iOS/tvOS device detection."));
+			Enable(false);
+			return;
+		}
+
         // get the list of devices UDID
         FPlatformProcess::ExecProcess(*LibimobileDeviceId, TEXT(""), &ReturnCode, &OutStdOut, &OutStdErr, NULL, true);
         if (OutStdOut.Len() == 0)
diff --git a/Engine/Source/Developer/IOS/IOSTargetPlatform/Private/IOSTargetDeviceOutput.h b/Engine/Source/Developer/IOS/IOSTargetPlatform/Private/IOSTargetDeviceOutput.h
--- a/Engine/Source/Developer/IOS/IOSTargetPlatform/Private/IOSTargetDeviceOutput.h
+++ b/Engine/Source/Developer/IOS/IOSTargetPlatform/Private/IOSTargetDeviceOutput.h
@@ -26,8 +26,7 @@ static FString GetLibImobileDeviceExe(const FString& ExeName)
 #if PLATFORM_WINDOWS
 	ToReturn += TEXT(".exe");
 #endif
-
-	return ToReturn;
+	return FPaths::FileExists(ToReturn) ? ToReturn : TEXT("");
 }
 
 class FIOSDeviceOutputReaderRunnable : public FRunnable
