Fixing up some style fallout from the FAutomationWindowStyle creation
#rnx

diff --git a/Engine/Source/Developer/AutomationWindow/Private/SAutomationGraphicalResultBox.cpp b/Engine/Source/Developer/AutomationWindow/Private/SAutomationGraphicalResultBox.cpp
--- a/Engine/Source/Developer/AutomationWindow/Private/SAutomationGraphicalResultBox.cpp
+++ b/Engine/Source/Developer/AutomationWindow/Private/SAutomationGraphicalResultBox.cpp
@@ -7,6 +7,7 @@
 #include "Widgets/Text/STextBlock.h"
 #include "Widgets/Layout/SGridPanel.h"
 #include "Styling/AppStyle.h"
+#include "AutomationWindowStyle.h"
 
 #define LOCTEXT_NAMESPACE "SAutomationGraphicalResultBox"
 
@@ -166,7 +167,7 @@ void SAutomationGraphicalResultBox::CreateWidgets()
 			.Padding( FMargin(1,3) )
 			[
 				SNew(STextBlock)
-				.TextStyle( FAppStyle::Get(), "Automation.ReportHeader" )
+				.TextStyle(FAutomationWindowStyle::Get(), "Automation.ReportHeader" )
 				.Text(FText::Format(LOCTEXT("AutomationGraphicalClusterHeader", "{Name}  -  {NumTests} Tests / {NumFails} Fails / {TotalTime} Seconds (Total) / {ParallelTime} Seconds (Parallel)"), ClusterArgs))
 			];
 
diff --git a/Engine/Source/Developer/IOS/IOSPlatformEditor/Private/IOSTargetSettingsCustomization.cpp b/Engine/Source/Developer/IOS/IOSPlatformEditor/Private/IOSTargetSettingsCustomization.cpp
--- a/Engine/Source/Developer/IOS/IOSPlatformEditor/Private/IOSTargetSettingsCustomization.cpp
+++ b/Engine/Source/Developer/IOS/IOSPlatformEditor/Private/IOSTargetSettingsCustomization.cpp
@@ -1538,11 +1538,11 @@ const FSlateBrush* FIOSTargetSettingsCustomization::GetProvisionStatus() const
 {
 	if( bProvisionInstalled )
 	{
-		return FAppStyle::GetBrush("Automation.Success");
+		return FAppStyle::GetBrush("Icons.Success");
 	}
 	else
 	{
-		return FAppStyle::GetBrush("Automation.Fail");
+		return FAppStyle::GetBrush("Icons.Error");
 	}
 }
 
@@ -1550,11 +1550,11 @@ const FSlateBrush* FIOSTargetSettingsCustomization::GetCertificateStatus() const
 {
 	if( bCertificateInstalled )
 	{
-		return FAppStyle::GetBrush("Automation.Success");
+		return FAppStyle::GetBrush("Icons.Success");
 	}
 	else
 	{
-		return FAppStyle::GetBrush("Automation.Fail");
+		return FAppStyle::GetBrush("Icons.Error");
 	}
 }
 
