- Removed a comment about a deleted tool

diff --git a/Engine/Source/Developer/IOS/IOSTargetPlatform/Private/IOSTargetPlatform.cpp b/Engine/Source/Developer/IOS/IOSTargetPlatform/Private/IOSTargetPlatform.cpp
--- a/Engine/Source/Developer/IOS/IOSTargetPlatform/Private/IOSTargetPlatform.cpp
+++ b/Engine/Source/Developer/IOS/IOSTargetPlatform/Private/IOSTargetPlatform.cpp
@@ -130,8 +130,7 @@ bool FIOSTargetPlatform::IsSdkInstalled(bool bProjectHasCode, FString& OutTutori
 
 	// The logic here is to assume the correct Apple dll does not exist and then check the various locations it could be in, setting this to true when it is found
 	// Code is structured for clarity not performance
-	// See Engine\Source\Programs\IOS\MobileDeviceInterface\MobileDevice.cs for reference
-	bool biOSSDKInstalled = false; 
+	bool biOSSDKInstalled = false;
 
 	HKEY hKey;
 	TCHAR dllPath[256];
