Fix preflight-caught instances of TObjectPtr<T> used with T that is not derived from UObject.  A subsequent change will be disallowing this at compile time.

diff --git a/Engine/Source/Developer/IOS/IOSTargetPlatform/Private/IOSDeviceHelper.cpp b/Engine/Source/Developer/IOS/IOSTargetPlatform/Private/IOSDeviceHelper.cpp
--- a/Engine/Source/Developer/IOS/IOSTargetPlatform/Private/IOSDeviceHelper.cpp
+++ b/Engine/Source/Developer/IOS/IOSTargetPlatform/Private/IOSDeviceHelper.cpp
@@ -337,7 +337,7 @@ private:
 
         for (int32 Index = 0; Index < ParsedDevices.Num(); ++Index)
         {
-			TObjectPtr<LibIMobileDevice> Found = CachedDevices.FindByPredicate(
+			LibIMobileDevice* Found = CachedDevices.FindByPredicate(
 				[&](LibIMobileDevice Element) {
 					return Element.DeviceID == ParsedDevices[Index].DeviceID;
 				});
