Fix a crash bug of enable clustered reflection on IOS.

diff --git a/Engine/Source/Developer/IOS/IOSTargetPlatform/Private/IOSTargetPlatform.cpp b/Engine/Source/Developer/IOS/IOSTargetPlatform/Private/IOSTargetPlatform.cpp
--- a/Engine/Source/Developer/IOS/IOSTargetPlatform/Private/IOSTargetPlatform.cpp
+++ b/Engine/Source/Developer/IOS/IOSTargetPlatform/Private/IOSTargetPlatform.cpp
@@ -546,7 +546,7 @@ void FIOSTargetPlatform::GetReflectionCaptureFormats( TArray<FName>& OutFormats
 {
 	const bool bMobileDeferredShading = (MobileShadingPath == 1);
 
-	if (SupportsMetalMRT() || bMobileDeferredShading)
+	if (SupportsMetalMRT() || bMobileDeferredShading || bMobileForwardEnableClusteredReflections)
 	{
 		OutFormats.Add(FName(TEXT("FullHDR")));
 	}
