PR #9956: Fix the casing on a couple of #includes. (Contributed by [REDACTED_PERSON])

diff --git a/Engine/Plugins/Media/ElectraUtil/Source/ElectraBase/Public/MediaVideoDecoderOutputApple.h b/Engine/Plugins/Media/ElectraUtil/Source/ElectraBase/Public/MediaVideoDecoderOutputApple.h
--- a/Engine/Plugins/Media/ElectraUtil/Source/ElectraBase/Public/MediaVideoDecoderOutputApple.h
+++ b/Engine/Plugins/Media/ElectraUtil/Source/ElectraBase/Public/MediaVideoDecoderOutputApple.h
@@ -5,7 +5,7 @@
 
 #include "Templates/RefCounting.h"
 #include "Containers/Array.h"
-#include "Pixelformat.h"
+#include "PixelFormat.h"
 
 class FVideoDecoderOutputApple : public FVideoDecoderOutput
 {
diff --git a/Engine/Source/Developer/IOS/IOSPlatformEditor/Private/IOSCustomIconProjectBuildMutatorFeature.cpp b/Engine/Source/Developer/IOS/IOSPlatformEditor/Private/IOSCustomIconProjectBuildMutatorFeature.cpp
--- a/Engine/Source/Developer/IOS/IOSPlatformEditor/Private/IOSCustomIconProjectBuildMutatorFeature.cpp
+++ b/Engine/Source/Developer/IOS/IOSPlatformEditor/Private/IOSCustomIconProjectBuildMutatorFeature.cpp
@@ -8,7 +8,7 @@
 #include "Interfaces/ITargetPlatformManagerModule.h"
 #include "GenericPlatform/GenericPlatformFile.h"
 #include "Misc/CoreMisc.h"
-#include "Interfaces/IturnkeySupportModule.h"
+#include "Interfaces/ITurnkeySupportModule.h"
 
 static bool RequiresBuild()
 {
