Secondary remote Mac

diff --git a/Engine/Config/BaseEngine.ini b/Engine/Config/BaseEngine.ini
--- a/Engine/Config/BaseEngine.ini
+++ b/Engine/Config/BaseEngine.ini
@@ -2731,11 +2731,16 @@ bSupportsIPhone=True
 AdditionalPlistData=
 RemoteServerName=
 RemoteServerOverrideBuildPath=
+SecondaryRemoteServerName=
+SecondaryRemoteServerOverrideBuildPath=
 RSyncUsername=
+SecondaryRSyncUsername=
 SSHPrivateKeyOverridePath=
+SecondarySSHPrivateKeyOverridePath=
 bEnableRemoteNotificationsSupport=False
 bEnableCloudKitSupport=False
 bRunAsCurrentUser=False
+bSupportSecondaryMac=False
 IOSCloudKitSyncStrategy=None
 bGenerateCrashReportSymbols=false
 bAutomaticSigning=false
diff --git a/Engine/Source/Developer/IOS/IOSPlatformEditor/Private/IOSTargetSettingsCustomization.cpp b/Engine/Source/Developer/IOS/IOSPlatformEditor/Private/IOSTargetSettingsCustomization.cpp
--- a/Engine/Source/Developer/IOS/IOSPlatformEditor/Private/IOSTargetSettingsCustomization.cpp
+++ b/Engine/Source/Developer/IOS/IOSPlatformEditor/Private/IOSTargetSettingsCustomization.cpp
@@ -149,6 +149,8 @@ void FIOSTargetSettingsCustomization::CustomizeDetails(IDetailLayoutBuilder& Det
 	BuildIconSection(DetailLayout);
 
 	BuildRemoteBuildingSection(DetailLayout);
+    
+    BuildSecondaryRemoteMacBuildingSection(DetailLayout);
 
 	AudioPluginWidgetManager.BuildAudioCategory(DetailLayout, FString(TEXT("IOS")));
 
@@ -941,31 +943,31 @@ void FIOSTargetSettingsCustomization::BuildRemoteBuildingSection(IDetailLayoutBu
 		[
 			SNew(SHorizontalBox)
 			+ SHorizontalBox::Slot()
-			.Padding(FMargin(0, 1, 0, 1))
-			.FillWidth(1.0f)
-			[
-				SNew(STextBlock)
-				.Text(LOCTEXT("RemoteServerNameLabel", "Remote Server Name"))
-				.Font(DetailLayout.GetDetailFont())
-			]
+		.Padding(FMargin(0, 1, 0, 1))
+		.FillWidth(1.0f)
+		[
+			SNew(STextBlock)
+			.Text(LOCTEXT("RemoteServerNameLabel", "Remote Server Name"))
+		.Font(DetailLayout.GetDetailFont())
 		]
-		.ValueContent()
+		]
+	.ValueContent()
 		.MinDesiredWidth(150.0f)
 		[
 			SNew(SHorizontalBox)
 			+ SHorizontalBox::Slot()
-			.Padding(FMargin(0.0f, 8.0f))
-			[
-				SNew(SEditableTextBox)
-				.IsEnabled(this, &FIOSTargetSettingsCustomization::IsImportEnabled)
-				.Text(this, &FIOSTargetSettingsCustomization::GetBundleText, RemoteServerNamePropertyHandle)
-				.Font(DetailLayout.GetDetailFont())
-				.SelectAllTextOnCommit(true)
-				.SelectAllTextWhenFocused(true)
-				.ClearKeyboardFocusOnCommit(false)
-				.ToolTipText(RemoteServerNamePropertyHandle->GetToolTipText())
-				.OnTextCommitted(this, &FIOSTargetSettingsCustomization::OnRemoteServerChanged, RemoteServerNamePropertyHandle)
-			]
+		.Padding(FMargin(0.0f, 8.0f))
+		[
+			SNew(SEditableTextBox)
+			.IsEnabled(this, &FIOSTargetSettingsCustomization::IsImportEnabled)
+		.Text(this, &FIOSTargetSettingsCustomization::GetBundleText, RemoteServerNamePropertyHandle)
+		.Font(DetailLayout.GetDetailFont())
+		.SelectAllTextOnCommit(true)
+		.SelectAllTextWhenFocused(true)
+		.ClearKeyboardFocusOnCommit(false)
+		.ToolTipText(RemoteServerNamePropertyHandle->GetToolTipText())
+		.OnTextCommitted(this, &FIOSTargetSettingsCustomization::OnRemoteServerChanged, RemoteServerNamePropertyHandle)
+		]
 
 		];
 
@@ -982,36 +984,36 @@ void FIOSTargetSettingsCustomization::BuildRemoteBuildingSection(IDetailLayoutBu
 	RSyncUsernamePropertyRow
 		.ToolTip(RSyncUsernamePropertyHandle->GetToolTipText())
 		.CustomWidget()
-			.NameContent()
-			[
-				SNew(SHorizontalBox)
-				+ SHorizontalBox::Slot()
-				.Padding(FMargin(0, 1, 0, 1))
-				.FillWidth(1.0f)
-				[
-					SNew(STextBlock)
-					.Text(LOCTEXT("RSyncUserNameLabel", "RSync User Name"))
-					.Font(DetailLayout.GetDetailFont())
-				]
-			]
-			.ValueContent()
-			.MinDesiredWidth(150.0f)
-			[
-				SNew(SHorizontalBox)
-				+ SHorizontalBox::Slot()
-				.Padding(FMargin(0.0f, 8.0f))
-				[
-					SNew(SEditableTextBox)
-					.IsEnabled(this, &FIOSTargetSettingsCustomization::IsImportEnabled)
-					.Text(this, &FIOSTargetSettingsCustomization::GetBundleText, RSyncUsernamePropertyHandle)
-					.Font(DetailLayout.GetDetailFont())
-					.SelectAllTextOnCommit(true)
-					.SelectAllTextWhenFocused(true)
-					.ClearKeyboardFocusOnCommit(false)
-					.ToolTipText(RSyncUsernamePropertyHandle->GetToolTipText())
-					.OnTextCommitted(this, &FIOSTargetSettingsCustomization::OnRemoteServerChanged, RSyncUsernamePropertyHandle)
-				]
-			];
+		.NameContent()
+		[
+			SNew(SHorizontalBox)
+			+ SHorizontalBox::Slot()
+		.Padding(FMargin(0, 1, 0, 1))
+		.FillWidth(1.0f)
+		[
+			SNew(STextBlock)
+			.Text(LOCTEXT("RSyncUserNameLabel", "RSync User Name"))
+		.Font(DetailLayout.GetDetailFont())
+		]
+		]
+	.ValueContent()
+		.MinDesiredWidth(150.0f)
+		[
+			SNew(SHorizontalBox)
+			+ SHorizontalBox::Slot()
+		.Padding(FMargin(0.0f, 8.0f))
+		[
+			SNew(SEditableTextBox)
+			.IsEnabled(this, &FIOSTargetSettingsCustomization::IsImportEnabled)
+		.Text(this, &FIOSTargetSettingsCustomization::GetBundleText, RSyncUsernamePropertyHandle)
+		.Font(DetailLayout.GetDetailFont())
+		.SelectAllTextOnCommit(true)
+		.SelectAllTextWhenFocused(true)
+		.ClearKeyboardFocusOnCommit(false)
+		.ToolTipText(RSyncUsernamePropertyHandle->GetToolTipText())
+		.OnTextCommitted(this, &FIOSTargetSettingsCustomization::OnRemoteServerChanged, RSyncUsernamePropertyHandle)
+		]
+		];
 
 
 	// Add Remote Server Override Build Path 
@@ -1042,28 +1044,170 @@ void FIOSTargetSettingsCustomization::BuildRemoteBuildingSection(IDetailLayoutBu
 		[
 			SNew(SVerticalBox)
 			+ SVerticalBox::Slot()
-				.AutoHeight()
-				[
-					SNew(SHorizontalBox)
-					+ SHorizontalBox::Slot()
-					.Padding(FMargin(0, 5, 0, 10))
-					.AutoWidth()
-					[
-						SNew(SButton)
-						.HAlign(HAlign_Center)
-						.VAlign(VAlign_Center)
-						.OnClicked(this, &FIOSTargetSettingsCustomization::OnGenerateSSHKey)
-						.IsEnabled(this, &FIOSTargetSettingsCustomization::IsImportEnabled)
-						[
-							SNew(STextBlock)
-							.Text(GenerateSSHText)
-						]
-					]
-				]
+		.AutoHeight()
+		[
+			SNew(SHorizontalBox)
+			+ SHorizontalBox::Slot()
+		.Padding(FMargin(0, 5, 0, 10))
+		.AutoWidth()
+		[
+			SNew(SButton)
+			.HAlign(HAlign_Center)
+		.VAlign(VAlign_Center)
+		.OnClicked(this, &FIOSTargetSettingsCustomization::OnGenerateSSHKey, true)
+		.IsEnabled(this, &FIOSTargetSettingsCustomization::IsImportEnabled)
+		[
+			SNew(STextBlock)
+			.Text(GenerateSSHText)
+		]
+		]
+		]
 		];
 #endif
 }
 
+void FIOSTargetSettingsCustomization::BuildSecondaryRemoteMacBuildingSection(IDetailLayoutBuilder& DetailLayout)
+{
+#if PLATFORM_WINDOWS
+	IDetailCategoryBuilder& BuildCategory = DetailLayout.EditCategory(TEXT("Build"));
+
+	// Sub group we wish to add remote building options to.
+	FText RemoteBuildingGroupName = LOCTEXT("RemoteBuildingGroupName", "Secondary Remote Build Options");
+	IDetailGroup& RemoteBuildingGroup = BuildCategory.AddGroup(*RemoteBuildingGroupName.ToString(), RemoteBuildingGroupName, false);
+
+	// Remote Server Name Property
+	TSharedRef<IPropertyHandle> RemoteServerNamePropertyHandle = DetailLayout.GetProperty(GET_MEMBER_NAME_CHECKED(UIOSRuntimeSettings, SecondaryRemoteServerName));
+
+	IDetailPropertyRow& RemoteServerNamePropertyRow = RemoteBuildingGroup.AddPropertyRow(RemoteServerNamePropertyHandle);
+	RemoteServerNamePropertyRow
+		.ToolTip(RemoteServerNamePropertyHandle->GetToolTipText())
+		.CustomWidget()
+		.NameContent()
+		[
+			SNew(SHorizontalBox)
+			+ SHorizontalBox::Slot()
+		.Padding(FMargin(0, 1, 0, 1))
+		.FillWidth(1.0f)
+		[
+			SNew(STextBlock)
+			.Text(LOCTEXT("RemoteServerNameLabel", "Secondary Remote Server Name"))
+		.Font(DetailLayout.GetDetailFont())
+		]
+		]
+	.ValueContent()
+		.MinDesiredWidth(150.0f)
+		[
+			SNew(SHorizontalBox)
+			+ SHorizontalBox::Slot()
+		.Padding(FMargin(0.0f, 8.0f))
+		[
+			SNew(SEditableTextBox)
+			.IsEnabled(this, &FIOSTargetSettingsCustomization::IsImportEnabled)
+		.Text(this, &FIOSTargetSettingsCustomization::GetBundleText, RemoteServerNamePropertyHandle)
+		.Font(DetailLayout.GetDetailFont())
+		.SelectAllTextOnCommit(true)
+		.SelectAllTextWhenFocused(true)
+		.ClearKeyboardFocusOnCommit(false)
+		.ToolTipText(RemoteServerNamePropertyHandle->GetToolTipText())
+		.OnTextCommitted(this, &FIOSTargetSettingsCustomization::OnRemoteServerChanged, RemoteServerNamePropertyHandle)
+		]
+
+		];
+
+
+
+	// Add Use RSync Property
+	TSharedRef<IPropertyHandle> UseRSyncPropertyHandle = DetailLayout.GetProperty(GET_MEMBER_NAME_CHECKED(UIOSRuntimeSettings, bUseRSync));
+	BuildCategory.AddProperty(UseRSyncPropertyHandle)
+		.Visibility(EVisibility::Hidden);
+
+	// Add RSync Username Property
+	TSharedRef<IPropertyHandle> RSyncUsernamePropertyHandle = DetailLayout.GetProperty(GET_MEMBER_NAME_CHECKED(UIOSRuntimeSettings, SecondaryRSyncUsername));
+	IDetailPropertyRow& RSyncUsernamePropertyRow = RemoteBuildingGroup.AddPropertyRow(RSyncUsernamePropertyHandle);
+	RSyncUsernamePropertyRow
+		.ToolTip(RSyncUsernamePropertyHandle->GetToolTipText())
+		.CustomWidget()
+		.NameContent()
+		[
+			SNew(SHorizontalBox)
+			+ SHorizontalBox::Slot()
+		.Padding(FMargin(0, 1, 0, 1))
+		.FillWidth(1.0f)
+		[
+			SNew(STextBlock)
+			.Text(LOCTEXT("RSyncUserNameLabel", "Secondary RSync User Name"))
+		.Font(DetailLayout.GetDetailFont())
+		]
+		]
+	.ValueContent()
+		.MinDesiredWidth(150.0f)
+		[
+			SNew(SHorizontalBox)
+			+ SHorizontalBox::Slot()
+		.Padding(FMargin(0.0f, 8.0f))
+		[
+			SNew(SEditableTextBox)
+			.IsEnabled(this, &FIOSTargetSettingsCustomization::IsImportEnabled)
+		.Text(this, &FIOSTargetSettingsCustomization::GetBundleText, RSyncUsernamePropertyHandle)
+		.Font(DetailLayout.GetDetailFont())
+		.SelectAllTextOnCommit(true)
+		.SelectAllTextWhenFocused(true)
+		.ClearKeyboardFocusOnCommit(false)
+		.ToolTipText(RSyncUsernamePropertyHandle->GetToolTipText())
+		.OnTextCommitted(this, &FIOSTargetSettingsCustomization::OnRemoteServerChanged, RSyncUsernamePropertyHandle)
+		]
+		];
+
+
+	// Add Remote Server Override Build Path 
+	TSharedRef<IPropertyHandle> RemoteServerOverrideBuildPathPropertyHandle = DetailLayout.GetProperty(GET_MEMBER_NAME_CHECKED(UIOSRuntimeSettings, SecondaryRemoteServerOverrideBuildPath));
+	IDetailPropertyRow& RemoteServerOverrideBuildPathPropertyRow = RemoteBuildingGroup.AddPropertyRow(RemoteServerOverrideBuildPathPropertyHandle);
+	RemoteServerOverrideBuildPathPropertyRow
+		.ToolTip(RemoteServerOverrideBuildPathPropertyHandle->GetToolTipText());
+
+	// Add existing SSH path label.
+	TSharedRef<IPropertyHandle> SSHPrivateKeyLocationPropertyHandle = DetailLayout.GetProperty(GET_MEMBER_NAME_CHECKED(UIOSRuntimeSettings, SecondarySSHPrivateKeyLocation));
+	IDetailPropertyRow& SSHPrivateKeyLocationPropertyRow = RemoteBuildingGroup.AddPropertyRow(SSHPrivateKeyLocationPropertyHandle);
+	SSHPrivateKeyLocationPropertyRow
+		.ToolTip(SSHPrivateKeyLocationPropertyHandle->GetToolTipText());
+
+	// cwRsync path
+	TSharedRef<IPropertyHandle> CwRsyncOverridePathPropertyHandle = DetailLayout.GetProperty(GET_MEMBER_NAME_CHECKED(UIOSRuntimeSettings, SecondaryCwRsyncInstallPath));
+	IDetailPropertyRow& CwRsyncOverridePathPropertyRow = RemoteBuildingGroup.AddPropertyRow(CwRsyncOverridePathPropertyHandle);
+
+	const FText GenerateSSHText = LOCTEXT("GenerateSSHKey", "Secondary Generate SSH Key");
+
+	// Add a generate key button
+	RemoteBuildingGroup.AddWidgetRow()
+		.FilterString(GenerateSSHText)
+		.WholeRowWidget
+		.MinDesiredWidth(0.f)
+		.MaxDesiredWidth(0.f)
+		.HAlign(HAlign_Fill)
+		[
+			SNew(SVerticalBox)
+			+ SVerticalBox::Slot()
+		.AutoHeight()
+		[
+			SNew(SHorizontalBox)
+			+ SHorizontalBox::Slot()
+		.Padding(FMargin(0, 5, 0, 10))
+		.AutoWidth()
+		[
+			SNew(SButton)
+			.HAlign(HAlign_Center)
+		.VAlign(VAlign_Center)
+		.OnClicked(this, &FIOSTargetSettingsCustomization::OnGenerateSSHKey, false)
+		.IsEnabled(this, &FIOSTargetSettingsCustomization::IsImportEnabled)
+		[
+			SNew(STextBlock)
+			.Text(GenerateSSHText)
+		]
+		]
+		]
+		];
+#endif
+}
 
 void FIOSTargetSettingsCustomization::BuildIconSection(IDetailLayoutBuilder& DetailLayout)
 {
@@ -1469,28 +1613,47 @@ FReply FIOSTargetSettingsCustomization::OnCertificateRequestClicked()
 	return FReply::Handled();
 }
 
-FReply FIOSTargetSettingsCustomization::OnGenerateSSHKey()
+FReply FIOSTargetSettingsCustomization::OnGenerateSSHKey(bool IsPrimary)
 {
 	// see if the key is already generated
 	const UIOSRuntimeSettings& Settings = *GetDefault<UIOSRuntimeSettings>();
 
 	FString RemoteServerAddress;
 	FString RemoteServerPort;
+	FString RSyncUsername;
 	int32	colonIndex;
 
-	if(Settings.RemoteServerName.FindChar(':', colonIndex))
+	if (IsPrimary)
 	{
-		RemoteServerAddress = Settings.RemoteServerName.Left(colonIndex);
-		RemoteServerPort = Settings.RemoteServerName.RightChop(colonIndex + 1);
+		if (Settings.RemoteServerName.FindChar(':', colonIndex))
+		{
+			RemoteServerAddress = Settings.RemoteServerName.Left(colonIndex);
+			RemoteServerPort = Settings.RemoteServerName.RightChop(colonIndex + 1);
+		}
+		else
+		{
+			RemoteServerAddress = Settings.RemoteServerName;
+			RemoteServerPort = "22";
+		}
+		RSyncUsername = Settings.RSyncUsername;
 	}
 	else
 	{
-		RemoteServerAddress = Settings.RemoteServerName;
-		RemoteServerPort = "22";
+		if (Settings.SecondaryRemoteServerName.FindChar(':', colonIndex))
+		{
+			RemoteServerAddress = Settings.SecondaryRemoteServerName.Left(colonIndex);
+			RemoteServerPort = Settings.SecondaryRemoteServerName.RightChop(colonIndex + 1);
+		}
+		else
+		{
+			RemoteServerAddress = Settings.SecondaryRemoteServerName;
+			RemoteServerPort = "22";
+		}
+		RSyncUsername = Settings.SecondaryRSyncUsername;
 	}
 
 	FString Path = FPlatformMisc::GetEnvironmentVariable(TEXT("APPDATA"));
-	FString Destination = FString::Printf(TEXT("%s\\Unreal Engine\\UnrealBuildTool\\SSHKeys\\%s\\%s\\RemoteToolChainPrivate.key"), *Path, *RemoteServerAddress, *(Settings.RSyncUsername));
+	FString Destination = FString::Printf(TEXT("%s\\Unreal Engine\\UnrealBuildTool\\SSHKeys\\%s\\%s\\RemoteToolChainPrivate.key"), *Path, *RemoteServerAddress, *RSyncUsername);
 	if (FPaths::FileExists(Destination))
 	{
 		FString MessagePrompt = FString::Printf(TEXT("An SSH Key already exists.  Do you want to replace this key?"));
diff --git a/Engine/Source/Developer/IOS/IOSPlatformEditor/Private/IOSTargetSettingsCustomization.h b/Engine/Source/Developer/IOS/IOSPlatformEditor/Private/IOSTargetSettingsCustomization.h
--- a/Engine/Source/Developer/IOS/IOSPlatformEditor/Private/IOSTargetSettingsCustomization.h
+++ b/Engine/Source/Developer/IOS/IOSPlatformEditor/Private/IOSTargetSettingsCustomization.h
@@ -118,7 +118,8 @@ private:
 	void BuildPListSection(IDetailLayoutBuilder& DetailLayout);
 	void BuildIconSection(IDetailLayoutBuilder& DetailLayout);
     void BuildRemoteBuildingSection(IDetailLayoutBuilder& DetailLayout);
-
+    void BuildSecondaryRemoteMacBuildingSection(IDetailLayoutBuilder& DetailLayout);
+    
 	// Navigates to the plist in explorer or finder
 	FReply OpenPlistFolder();
 
@@ -140,7 +141,7 @@ private:
 	FReply OnGenerateImageClicked(const FString SourceImagePath, const FString TargetImagePath, FIntPoint IconRequiredSize);
 
 	// ssh key request
-	FReply OnGenerateSSHKey();
+	FReply OnGenerateSSHKey(bool IsPrimary);
 
 	// Get the image to display for the provision status
 	const FSlateBrush* GetProvisionStatus() const;
diff --git a/Engine/Source/Developer/IOS/IOSTargetPlatform/Private/IOSTargetPlatform.cpp b/Engine/Source/Developer/IOS/IOSTargetPlatform/Private/IOSTargetPlatform.cpp
--- a/Engine/Source/Developer/IOS/IOSTargetPlatform/Private/IOSTargetPlatform.cpp
+++ b/Engine/Source/Developer/IOS/IOSTargetPlatform/Private/IOSTargetPlatform.cpp
@@ -282,9 +282,13 @@ int32 FIOSTargetPlatform::CheckRequirements(bool bProjectHasCode, EBuildConfigur
 	FString CmdExe = FPaths::ConvertRelativePathToFull(FPaths::EngineDir() / TEXT("Binaries/DotNET/IOS/IPhonePackager.exe"));
 	FString CommandLine = FString::Printf(TEXT("Validate Engine -project \"%s\" -bundlename \"%s\" %s"), *ProjectPath, *(BundleIdentifier), (bForDistribtion ? TEXT("-distribution") : TEXT("")) );
 	FString RemoteServerName;
+	FString SecondaryRemoteServerName;
 	FString RSyncUsername;
+	FString SecondaryRSyncUsername;
 	GConfig->GetString(TEXT("/Script/IOSRuntimeSettings.IOSRuntimeSettings"), TEXT("RemoteServerName"), RemoteServerName, GEngineIni);
 	GConfig->GetString(TEXT("/Script/IOSRuntimeSettings.IOSRuntimeSettings"), TEXT("RSyncUsername"), RSyncUsername, GEngineIni);
+	GConfig->GetString(TEXT("/Script/IOSRuntimeSettings.IOSRuntimeSettings"), TEXT("SecondaryRemoteServerName"), SecondaryRemoteServerName, GEngineIni);
+	GConfig->GetString(TEXT("/Script/IOSRuntimeSettings.IOSRuntimeSettings"), TEXT("SecondaryRSyncUsername"), SecondaryRSyncUsername, GEngineIni);
 	if (RemoteServerName.Len() == 0 || RSyncUsername.Len() == 0)
 	{
 		bReadyToBuild |= ETargetPlatformReadyStatus::RemoveServerNameEmpty;
diff --git a/Engine/Source/Editor/TurnkeySupport/Private/TurnkeySupportModule.cpp b/Engine/Source/Editor/TurnkeySupport/Private/TurnkeySupportModule.cpp
--- a/Engine/Source/Editor/TurnkeySupport/Private/TurnkeySupportModule.cpp
+++ b/Engine/Source/Editor/TurnkeySupport/Private/TurnkeySupportModule.cpp
@@ -513,15 +513,30 @@ public:
 		}
         else if (Mode == EPrepareContentMode::PrepareForDebugging)
          {
-#if PLATFORM_MAC || PLATFORM_WINDOWS
+#if  PLATFORM_WINDOWS
              if (IniPlatformName == TEXT("IOS") || IniPlatformName == TEXT("TvOS"))
              {
-                 FString CommandLine = FString::Printf(TEXT("WrangleContentForDebugging -platform=%s -ProjectFilePath=%s"), *IniPlatformName.ToString().ToLower(), *ProjectPath);
-              
-                 FTurnkeyEditorSupport::RunUAT(CommandLine, PlatformInfo->DisplayName, ContentPrepDescription, ContentPrepTaskName, ContentPrepIcon, &AnalyticsParamArray);
-             
-                 return;
+				 bool bSupportSecondaryMac = GConfig->GetBool(TEXT("/Script/IOSRuntimeSettings.IOSRuntimeSettings"), TEXT("bSupportSecondaryMac"), bSupportSecondaryMac, GEngineIni);
+				 if (bSupportSecondaryMac)
+				 {
+					 FString OtherCommandLine = FString::Printf(TEXT("SetSecondaryRemoteMac -platform=%s -ProjectFilePath=%s"), *IniPlatformName.ToString().ToLower(), *ProjectPath);
+					 FTurnkeyEditorSupport::RunUAT(OtherCommandLine, PlatformInfo->DisplayName, ContentPrepDescription, ContentPrepTaskName, ContentPrepIcon, &AnalyticsParamArray);
+				 }
+				 else
+				 {
+					 FString CommandLine = FString::Printf(TEXT("WrangleContentForDebugging -platform=%s -ProjectFilePath=%s"), *IniPlatformName.ToString().ToLower(), *ProjectPath);
+					 FTurnkeyEditorSupport::RunUAT(CommandLine, PlatformInfo->DisplayName, ContentPrepDescription, ContentPrepTaskName, ContentPrepIcon, &AnalyticsParamArray);
+				 }
+
+				 return;
              }
+#elif PLATFORM_MAC
+		 if (IniPlatformName == TEXT("IOS") || IniPlatformName == TEXT("TvOS"))
+		 {
+			 FString CommandLine = FString::Printf(TEXT("WrangleContentForDebugging -platform=%s -ProjectFilePath=%s"), *IniPlatformName.ToString().ToLower(), *ProjectPath);
+			 FTurnkeyEditorSupport::RunUAT(CommandLine, PlatformInfo->DisplayName, ContentPrepDescription, ContentPrepTaskName, ContentPrepIcon, &AnalyticsParamArray);
+			 return;
+		 }
 #endif
  		}
 		else if (Mode == EPrepareContentMode::CookOnly)
@@ -1114,34 +1129,7 @@ static void MakeTurnkeyPlatformMenu(UToolMenu* ToolMenu, FName IniPlatformName,
 			)
 		);
         // add platforms here when they have their own Prepare for debug flow
-#if PLATFORM_MAC
-        if (IniPlatformName.ToString() == "IOS")
-        {
-            Section.AddMenuEntry(
-                                 NAME_None,
-                                 LOCTEXT("Turnkey_PrepareForDebugging", "Prepare For Debugging"),
-                                 LOCTEXT("TurnkeyTooltip_PrepareForDebuggingIOS", "Prepare this project for debugging. Expects an IPA package with the same name as the project file in the Binaries/IOS/ folder."),
-                                 FSlateIcon(),
-                                 FUIAction(
-                                           FExecuteAction::CreateStatic(&FTurnkeySupportCallbacks::CookOrPackage, IniPlatformName, EPrepareContentMode::PrepareForDebugging),
-                                           FCanExecuteAction::CreateStatic(&FTurnkeySupportCallbacks::CanCookOrPackage, IniPlatformName, EPrepareContentMode::PrepareForDebugging)
-                                           )
-                                 );
-        }
-        if (IniPlatformName.ToString() == "TvOS")
-        {
-            Section.AddMenuEntry(
-                                 NAME_None,
-                                 LOCTEXT("Turnkey_PrepareForDebugging", "Prepare For Debugging"),
-                                 LOCTEXT("TurnkeyTooltip_PrepareForDebuggingTvOS", "Prepare this project for debugging. Expects an IPA package with the same name as the project file in the Binaries/TVOS/ folder."),
-                                 FSlateIcon(),
-                                 FUIAction(
-                                           FExecuteAction::CreateStatic(&FTurnkeySupportCallbacks::CookOrPackage, IniPlatformName, EPrepareContentMode::PrepareForDebugging),
-                                           FCanExecuteAction::CreateStatic(&FTurnkeySupportCallbacks::CanCookOrPackage, IniPlatformName, EPrepareContentMode::PrepareForDebugging)
-                                           )
-                                 );
-        }
-#elif PLATFORM_WINDOWS
+#if PLATFORM_WINDOWS || PLATFORM_MAC
 		if (IniPlatformName.ToString() == "IOS" || IniPlatformName.ToString() == "TvOS")
 		{
 
@@ -1150,7 +1138,7 @@ static void MakeTurnkeyPlatformMenu(UToolMenu* ToolMenu, FName IniPlatformName,
 			Section.AddMenuEntry(
 				NAME_None,
 				LOCTEXT("Turnkey_PrepareForDebugging", "Prepare For Debugging"),
-				FText::Format(LOCTEXT("TurnkeyTooltip_PrepareForDebugging", "Expects a working remote toolchain and an IPA package with the same name as the project file in the Binaries/ {0} / folder."), FText::FromString(IniPlatformName.ToString().ToUpper())),
+				FText::Format(LOCTEXT("TurnkeyTooltip_PrepareForDebugging", "Expects a working remote toolchain and an IPA package with the same name as the project file in the Binaries/{0}/ folder."), FText::FromString(IniPlatformName.ToString().ToUpper())),
 				FSlateIcon(),
 				FUIAction(
 					FExecuteAction::CreateStatic(&FTurnkeySupportCallbacks::CookOrPackage, IniPlatformName, EPrepareContentMode::PrepareForDebugging),
diff --git a/Engine/Source/Programs/AutomationTool/AutomationUtils/Platform.cs b/Engine/Source/Programs/AutomationTool/AutomationUtils/Platform.cs
--- a/Engine/Source/Programs/AutomationTool/AutomationUtils/Platform.cs
+++ b/Engine/Source/Programs/AutomationTool/AutomationUtils/Platform.cs
@@ -873,6 +873,11 @@ namespace AutomationTool
 			LogError("Not implemented for the %s platform.", ClientPlatform);
 		}
 
+		public virtual void SetSecondaryRemoteMac(string ProjectFilePath, string ClientPlatform)
+		{
+			LogError("Not implemented for this platform.");
+		}
+
 		// let the platform set the exe extension if it chooses (otherwise, use
 		// the switch statement in GetExeExtension below)
 		protected virtual string GetPlatformExeExtension()
diff --git a/Engine/Source/Programs/AutomationTool/IOS/IOSPlatform.Automation.cs b/Engine/Source/Programs/AutomationTool/IOS/IOSPlatform.Automation.cs
--- a/Engine/Source/Programs/AutomationTool/IOS/IOSPlatform.Automation.cs
+++ b/Engine/Source/Programs/AutomationTool/IOS/IOSPlatform.Automation.cs
@@ -2294,6 +2294,12 @@ public class IOSPlatform : ApplePlatform
 		return DirectoryReference.Combine(RuntimeRoot, "cookeddata/" + RelativeProjectRootForStage.Name);
 	}
 
+	public override void SetSecondaryRemoteMac(string ProjectFilePath, string ClientPlatform)
+	{
+		PrepareForDebugging("", ProjectFilePath, ClientPlatform);
+		IOSExports.SetSecondaryRemoteMac(ClientPlatform, new FileReference(ProjectFilePath), Log.Logger);
+	}
+
 	public override void PrepareForDebugging(string SourcePackage, string ProjectFilePath, string ClientPlatform)
 	{
 		if (HostPlatform.Current.HostEditorPlatform == UnrealTargetPlatform.Win64 || HostPlatform.Current.HostEditorPlatform == UnrealTargetPlatform.Mac)
diff --git a/Engine/Source/Programs/AutomationTool/Scripts/SetSecondaryRemoteMac.Automation.cs b/Engine/Source/Programs/AutomationTool/Scripts/SetSecondaryRemoteMac.Automation.cs
new file mode 100644
--- /dev/null
+++ b/Engine/Source/Programs/AutomationTool/Scripts/SetSecondaryRemoteMac.Automation.cs
@@ -0,0 +1,38 @@
+// Copyright Epic Games, Inc. All Rights Reserved.
+
+using Microsoft.Extensions.Logging;
+using System.Collections.Generic;
+using AutomationTool;
+using UnrealBuildTool;
+using System.IO;
+using EpicGames.Core;
+
+namespace AutomationScripts
+{
+	public class SetSecondaryRemoteMac : BuildCommand
+	{
+		public override void ExecuteBuild()
+		{
+
+			if (ParseParam("usage"))
+			{
+				LogInformation("Arguments : \n -platform=<platform> \n -ProjectFilePath=<path/ProjectName.uproject> Mandatory. Path to the project you want to Debug. \n -SourcePackage=<path/PackageName.ipa> Optional. Specify the .ipa to wrangle data for debug from. When not specified, [ProjectPath]/Build/[IOS|TVOS]/[ProjectName].ipa will be used.");
+			}
+
+			FileReference ProjectFile = new FileReference(ParseParamValue("ProjectFilePath=", null));
+			var Params = new ProjectParams
+			(
+				Command: this,
+				RawProjectPath: ProjectFile
+			);
+
+			LogInformation("======= SetSecondaryRemoteMac - Start =======");
+
+			var SC = Project.CreateDeploymentContext(Params, false);
+
+			SC[0].StageTargetPlatform.SetSecondaryRemoteMac(ParseParamValue("ProjectFilePath=", null), ParseParamValue("platform=", null));
+
+			LogInformation("======= SetSecondaryRemoteMac - Done ========");
+		}
+	}
+}
\ No newline at end of file
diff --git a/Engine/Source/Programs/AutomationTool/Scripts/WrangleContentForDebugging.Automation.cs b/Engine/Source/Programs/AutomationTool/Scripts/WrangleContentForDebugging.Automation.cs
--- a/Engine/Source/Programs/AutomationTool/Scripts/WrangleContentForDebugging.Automation.cs
+++ b/Engine/Source/Programs/AutomationTool/Scripts/WrangleContentForDebugging.Automation.cs
@@ -5,7 +5,6 @@ using System.Collections.Generic;
 using AutomationTool;
 using UnrealBuildTool;
 using System.IO;
-//using System.Linq;
 using EpicGames.Core;
 
 namespace AutomationScripts
@@ -17,7 +16,7 @@ namespace AutomationScripts
 
 			if (ParseParam("usage"))
 			{
-				LogInformation("Arguments : \n -platform=<platform> \n -ProjectFilePath=<path/ProjectName.uproject> Mandatory. Path to a path to the roject you want to Wrangle Debug Data to. \n -SourcePackage=<path/PackageName.ipa> Optional. Specify the .ipa to wrangle data for debug from. When not specified, [ProjectPath]/Build/[IOS|TVOS]/[ProjectName].ipa will be used.");
+				LogInformation("Arguments : \n -platform=<platform> \n -ProjectFilePath=<path/ProjectName.uproject> Mandatory. Path to the project you want to debug. \n -SourcePackage=<path/PackageName.ipa> Optional. Specify the .ipa to wrangle data for debug from. When not specified, [ProjectPath]/Build/[IOS|TVOS]/[ProjectName].ipa will be used.");
 			}
 
 			FileReference ProjectFile = new FileReference(ParseParamValue("ProjectFilePath=", null));
@@ -37,4 +36,4 @@ namespace AutomationScripts
 			LogInformation("======= WrangleContentForDebugging - Done ========");
 		}
 	}
-}
\ No newline at end of file
+}
diff --git a/Engine/Source/Programs/UnrealBuildTool/Platform/IOS/IOSExports.cs b/Engine/Source/Programs/UnrealBuildTool/Platform/IOS/IOSExports.cs
--- a/Engine/Source/Programs/UnrealBuildTool/Platform/IOS/IOSExports.cs
+++ b/Engine/Source/Programs/UnrealBuildTool/Platform/IOS/IOSExports.cs
@@ -181,6 +181,20 @@ namespace UnrealBuildTool
 			}
 		}
 
+		/// <summary>
+		/// Set a secondary remote Mac to retrieve built data on a remote Mac.
+		/// </summary>
+		/// <returns></returns>
+		
+		public static void SetSecondaryRemoteMac(string ClientPlatform, FileReference ProjectFile, ILogger Logger)
+		{
+			RemoteMac Remote = new RemoteMac(ProjectFile, Logger, true, true);
+			Remote.RetrieveFilesGeneratedOnPrimaryMac(ProjectFile, Logger, ClientPlatform);
+
+			RemoteMac SecondaryRemote = new RemoteMac(ProjectFile, Logger, false);
+			SecondaryRemote.UploadToSecondaryMac(ProjectFile, Logger);
+		}
+
 		/// <summary>
 		/// Prepare the build and project on a Remote Mac to be able to debug an iOS or tvOS package built remotely
 		/// </summary>
diff --git a/Engine/Source/Programs/UnrealBuildTool/ToolChain/RemoteMac.cs b/Engine/Source/Programs/UnrealBuildTool/ToolChain/RemoteMac.cs
--- a/Engine/Source/Programs/UnrealBuildTool/ToolChain/RemoteMac.cs
+++ b/Engine/Source/Programs/UnrealBuildTool/ToolChain/RemoteMac.cs
@@ -130,7 +130,9 @@ namespace UnrealBuildTool
 		/// </summary>
 		/// <param name="ProjectFile">Project to read settings from</param>
 		/// <param name="Logger">Logger for output</param>
-		public RemoteMac(FileReference? ProjectFile, ILogger Logger)
+		/// <param name="bIsPrimary">Is the primary Remote Mac or not. True by default.</param>
+		/// <param name="bPrepareForSecondaryMac">Added arguments when preparing a secondary Mac for debug. False by default.</param>
+		public RemoteMac(FileReference? ProjectFile, ILogger Logger, bool bIsPrimary = true, bool bPrepareForSecondaryMac = false)
 		{
 			this.RsyncExe = FileReference.Combine(Unreal.EngineDirectory, "Extras", "ThirdPartyNotUE", "cwrsync", "bin", "rsync.exe");
 			this.SshExe = FileReference.Combine(Unreal.EngineDirectory, "Extras", "ThirdPartyNotUE", "cwrsync", "bin", "ssh.exe");
@@ -162,7 +164,8 @@ namespace UnrealBuildTool
 			{
 				// Read the server name
 				string IniServerName;
-				if (Ini.GetString("/Script/IOSRuntimeSettings.IOSRuntimeSettings", "RemoteServerName", out IniServerName) && !String.IsNullOrEmpty(IniServerName))
+
+				if ((bIsPrimary ? Ini.GetString("/Script/IOSRuntimeSettings.IOSRuntimeSettings", "RemoteServerName", out IniServerName) : Ini.GetString("/Script/IOSRuntimeSettings.IOSRuntimeSettings", "SecondaryRemoteServerName", out IniServerName)) && !String.IsNullOrEmpty(IniServerName))
 				{
 					this.ServerName = IniServerName;
 				}
@@ -173,7 +176,7 @@ namespace UnrealBuildTool
 
 				// Parse the username
 				string IniUserName;
-				if (Ini.GetString("/Script/IOSRuntimeSettings.IOSRuntimeSettings", "RSyncUsername", out IniUserName) && !String.IsNullOrEmpty(IniUserName))
+				if ((bIsPrimary ? Ini.GetString("/Script/IOSRuntimeSettings.IOSRuntimeSettings", "RSyncUsername", out IniUserName) : Ini.GetString("/Script/IOSRuntimeSettings.IOSRuntimeSettings", "SecondaryRSyncUsername", out IniUserName)) && !String.IsNullOrEmpty(IniUserName))
 				{
 					this.UserName = IniUserName;
 				}
@@ -202,7 +205,7 @@ namespace UnrealBuildTool
 
 			// Get the path to the SSH private key
 			string OverrideSshPrivateKeyPath;
-			if (Ini.GetString("/Script/IOSRuntimeSettings.IOSRuntimeSettings", "SSHPrivateKeyOverridePath", out OverrideSshPrivateKeyPath) && !String.IsNullOrEmpty(OverrideSshPrivateKeyPath))
+			if ((bIsPrimary ? Ini.GetString("/Script/IOSRuntimeSettings.IOSRuntimeSettings", "SSHPrivateKeyOverridePath", out OverrideSshPrivateKeyPath) : Ini.GetString("/Script/IOSRuntimeSettings.IOSRuntimeSettings", "SecondarySSHPrivateKeyOverridePath", out OverrideSshPrivateKeyPath)) && !String.IsNullOrEmpty(OverrideSshPrivateKeyPath))
 			{
 				SshPrivateKey = new FileReference(OverrideSshPrivateKeyPath);
 				if (!FileReference.Exists(SshPrivateKey))
@@ -266,15 +269,18 @@ namespace UnrealBuildTool
 			CommonRsyncArguments = new List<string>(BasicRsyncArguments);
 			CommonRsyncArguments.Add("--copy-links");
 			CommonRsyncArguments.Add("--recursive");
-			CommonRsyncArguments.Add("--delete"); // Delete anything not in the source directory
-			CommonRsyncArguments.Add("--delete-excluded"); // Delete anything not in the source directory
+			if (!bPrepareForSecondaryMac)
+			{
+				CommonRsyncArguments.Add("--delete"); // Delete anything not in the source directory
+				CommonRsyncArguments.Add("--delete-excluded"); // Delete anything not in the source directory
+			}
 			CommonRsyncArguments.Add("--times"); // Preserve modification times
 			CommonRsyncArguments.Add("--omit-dir-times"); // Ignore modification times for directories
 			CommonRsyncArguments.Add("--prune-empty-dirs"); // Remove empty directories from the file list
 
 			// Get the remote base directory
 			string RemoteServerOverrideBuildPath;
-			if (Ini.GetString("/Script/IOSRuntimeSettings.IOSRuntimeSettings", "RemoteServerOverrideBuildPath", out RemoteServerOverrideBuildPath) && !String.IsNullOrEmpty(RemoteServerOverrideBuildPath))
+			if ((bIsPrimary ? Ini.GetString("/Script/IOSRuntimeSettings.IOSRuntimeSettings", "RemoteServerOverrideBuildPath", out RemoteServerOverrideBuildPath) : Ini.GetString("/Script/IOSRuntimeSettings.IOSRuntimeSettings", "SecondaryRemoteServerOverrideBuildPath", out RemoteServerOverrideBuildPath)) && !String.IsNullOrEmpty(RemoteServerOverrideBuildPath))
 			{
 				RemoteBaseDir = String.Format("{0}/{1}", RemoteServerOverrideBuildPath.Trim().TrimEnd('/'), Environment.MachineName);
 			}
@@ -711,6 +717,12 @@ namespace UnrealBuildTool
 			DownloadFile(OutputFile, Logger);
 		}
 
+		/// <summary>
+		/// Prepare the remotely built iOS or tvOS client to de debugged on Mac
+		/// </summary>
+		/// <param name="CookedDataDirectory">Where the content has been previously cooked</param>
+		/// <param name="ProjectFile">The project file</param>
+		/// <param name="Logger">The logger</param>
 		public void PrepareToDebug(string CookedDataDirectory, FileReference ProjectFile, ILogger Logger)
         {
 			Logger.LogInformation("[Remote] Uploading the default uprojectdirs ...");
@@ -730,6 +742,49 @@ namespace UnrealBuildTool
 			Logger.LogInformation("An Xcode project file has been generated on your remote Mac at the above address. Please access it to debug your IPA.");
 		}
 
+		/// <summary>
+		/// Retrieves data from the primary Mac to the Windows machine
+		/// </summary>
+		/// <param name="ProjectFile">The project file</param>
+		/// <param name="Logger">The logger</param>
+		/// <param name="ClientPlatform">The platfrom for the client (typically iOS or tvOS)</param>
+		public void RetrieveFilesGeneratedOnPrimaryMac(FileReference ProjectFile, ILogger Logger, string ClientPlatform)
+        {
+			Logger.LogInformation("[Remote] Downloading files generated on the Primary Remote Mac ...");
+			DownloadDirectory(DirectoryReference.Combine(Unreal.EngineDirectory, "Intermediate"), Logger);
+			DownloadDirectory(DirectoryReference.Combine(Unreal.EngineDirectory, "Saved"), Logger);
+			
+			DownloadDirectory(DirectoryReference.Combine(Unreal.EngineDirectory, "Binaries", ClientPlatform), Logger);
+			DownloadDirectory(DirectoryReference.Combine(ProjectFile.Directory, "Intermediate"), Logger);
+			DownloadDirectory(DirectoryReference.Combine(ProjectFile.Directory, "Binaries", ClientPlatform), Logger);
+			DownloadDirectory(DirectoryReference.Combine(new DirectoryReference("UE5.xcworkspace")), Logger);
+
+		}
+
+		/// <summary>
+		/// Upload retrieved data from the primary Mac from the Windows machine to the secondary Mac
+		/// </summary>
+		/// <param name="ProjectFile">The project file</param>
+		/// <param name="Logger">The logger</param>
+		public void UploadToSecondaryMac(FileReference ProjectFile, ILogger Logger)
+		{
+			Logger.LogInformation("[Remote] Uploading workspace to the secondary Mac ...");
+
+			UploadWorkspace(DirectoryReference.Combine(Unreal.EngineDirectory, ".."), Logger);
+			UploadFile(new FileReference("Default.uprojectdirs"), Logger);
+			UploadDirectory(DirectoryReference.Combine(ProjectFile.Directory, "Intermediate"), Logger);
+			UploadDirectory(DirectoryReference.Combine(ProjectFile.Directory, "Binaries"), Logger);
+			UploadDirectory(DirectoryReference.Combine(ProjectFile.Directory, "Saved"), Logger);
+
+			Logger.LogInformation("[Remote] Generating an Xcode Project file...");
+			string CommandLine = EscapeShellArgument(GetRemotePath(Unreal.EngineDirectory));
+			CommandLine += "/Build/BatchFiles/Mac/GenerateProjectFiles.sh";
+			StringBuilder BuildCommandLine = new StringBuilder(CommandLine);	
+			Execute(GetRemotePath(Unreal.RootDirectory), CommandLine, Logger);
+
+			UploadDirectory(new DirectoryReference("UE5.xcworkspace"), Logger);
+		}
+
 		/// <summary>
 		/// Convers a remote path into local form
 		/// </summary>
diff --git a/Engine/Source/Runtime/IOS/IOSRuntimeSettings/Classes/IOSRuntimeSettings.h b/Engine/Source/Runtime/IOS/IOSRuntimeSettings/Classes/IOSRuntimeSettings.h
--- a/Engine/Source/Runtime/IOS/IOSRuntimeSettings/Classes/IOSRuntimeSettings.h
+++ b/Engine/Source/Runtime/IOS/IOSRuntimeSettings/Classes/IOSRuntimeSettings.h
@@ -248,6 +248,14 @@ public:
 	UPROPERTY(GlobalConfig, EditAnywhere, Category = "Build", meta = (ConfigHierarchyEditable))
 	FString RemoteServerName;
 
+	// The name or ip address of the remote mac which will be used to build IOS
+	UPROPERTY(GlobalConfig, EditAnywhere, Category = "Build", meta = (ConfigHierarchyEditable))
+	FString SecondaryRemoteServerName;
+
+	// Support a secondary remote mac ?
+	UPROPERTY(GlobalConfig, EditAnywhere, Category = "Build", meta = (DisplayName = "Supports a secondary remote mac to facilitate debug of iOS / tvOS builds"))
+	bool bSupportSecondaryMac;
+
 	// Enable the use of RSync for remote builds on a mac
 	UPROPERTY(GlobalConfig, EditAnywhere, Category = "Build", meta = (DisplayName = "Use RSync for building IOS", ConfigHierarchyEditable))
 	bool bUseRSync;
@@ -256,21 +264,41 @@ public:
 	UPROPERTY(GlobalConfig, EditAnywhere, Category = "Build", meta = (EditCondition = "bUseRSync", DisplayName = "Username on Remote Server.", ConfigHierarchyEditable))
 	FString RSyncUsername;
 
+	// The secondary mac users name which matches the SSH Private Key, for remote builds using RSync.
+	UPROPERTY(GlobalConfig, EditAnywhere, Category = "Build", meta = (EditCondition = "bUseRSync", DisplayName = "Username on Remote Server.", ConfigHierarchyEditable))
+	FString SecondaryRSyncUsername;
+
 	// Optional path on the remote mac where the build files will be copied. If blank, ~/UE5/Builds will be used.
 	UPROPERTY(GlobalConfig, EditAnywhere, Category = "Build", meta = (ConfigHierarchyEditable))
 	FString RemoteServerOverrideBuildPath;
 
+	// Optional path on the secondary remote mac where the build files will be copied. If blank, ~/UE5/Builds will be used.
+	UPROPERTY(GlobalConfig, EditAnywhere, Category = "Build", meta = (ConfigHierarchyEditable))
+	FString SecondaryRemoteServerOverrideBuildPath;
+
 	// The install directory of cwrsync.
 	UPROPERTY(GlobalConfig, EditAnywhere, Category = "Build", meta = (EditCondition = "bUseRSync", ConfigHierarchyEditable))
 	FIOSBuildResourceDirectory CwRsyncInstallPath;
 
+	// The install directory of cwrsync.
+	UPROPERTY(GlobalConfig, EditAnywhere, Category = "Build", meta = (EditCondition = "bUseRSync", ConfigHierarchyEditable))
+	FIOSBuildResourceDirectory SecondaryCwRsyncInstallPath;
+
 	// The existing location of an SSH Key found by Unreal Engine.
 	UPROPERTY(VisibleAnywhere, Category = "Build", meta = (DisplayName = "Found Existing SSH permissions file"))
 	FString SSHPrivateKeyLocation;
 
+	// The existing location of an SSH Key found by Unreal Engine.
+	UPROPERTY(VisibleAnywhere, Category = "Build", meta = (DisplayName = "Found Existing SSH permissions file for Secondary Mac"))
+	FString SecondarySSHPrivateKeyLocation;
+
 	// The path of the ssh permissions key to be used when connecting to the remote server.
 	UPROPERTY(GlobalConfig, EditAnywhere, Category = "Build", meta = (EditCondition = "bUseRSync", DisplayName = "Override existing SSH permissions file", ConfigHierarchyEditable))
 	FIOSBuildResourceFilePath SSHPrivateKeyOverridePath;
+
+	// The path of the ssh permissions key to be used when connecting to the remote server.
+	UPROPERTY(GlobalConfig, EditAnywhere, Category = "Build", meta = (EditCondition = "bUseRSync", DisplayName = "Override existing SSH permissions file for Secondary Mac", ConfigHierarchyEditable))
+	FIOSBuildResourceFilePath SecondarySSHPrivateKeyOverridePath;
     
     // Should the app be compatible with Multi-User feature on tvOS ?　If checked, the game will will shutdown with the typical exit flow.
     UPROPERTY(GlobalConfig, EditAnywhere, Category = "Build", meta = (DisplayName = "Support user switching on tvOS."))
diff --git a/Engine/Source/Runtime/IOS/IOSRuntimeSettings/Private/IOSRuntimeSettings.cpp b/Engine/Source/Runtime/IOS/IOSRuntimeSettings/Private/IOSRuntimeSettings.cpp
--- a/Engine/Source/Runtime/IOS/IOSRuntimeSettings/Private/IOSRuntimeSettings.cpp
+++ b/Engine/Source/Runtime/IOS/IOSRuntimeSettings/Private/IOSRuntimeSettings.cpp
@@ -37,6 +37,7 @@ UIOSRuntimeSettings::UIOSRuntimeSettings(const FObjectInitializer& ObjectInitial
 	bGeneratedSYMBundle = false;
 	bGenerateXCArchive = false;
 	bShipForBitcode = true;
+	bSupportSecondaryMac = false;
 	bUseRSync = true;
 	bCustomLaunchscreenStoryboard = false;
 	AdditionalPlistData = TEXT("");
