Fix localization warning

diff --git a/Engine/Source/Developer/IOS/IOSPlatformEditor/Private/IOSTargetSettingsCustomization.cpp b/Engine/Source/Developer/IOS/IOSPlatformEditor/Private/IOSTargetSettingsCustomization.cpp
--- a/Engine/Source/Developer/IOS/IOSPlatformEditor/Private/IOSTargetSettingsCustomization.cpp
+++ b/Engine/Source/Developer/IOS/IOSPlatformEditor/Private/IOSTargetSettingsCustomization.cpp
@@ -1072,7 +1072,7 @@ void FIOSTargetSettingsCustomization::BuildSecondaryRemoteMacBuildingSection(IDe
 	IDetailCategoryBuilder& BuildCategory = DetailLayout.EditCategory(TEXT("Build"));
 
 	// Sub group we wish to add remote building options to.
-	FText RemoteBuildingGroupName = LOCTEXT("RemoteBuildingGroupName", "Secondary Remote Build Options");
+	FText RemoteBuildingGroupName = LOCTEXT("SecondaryRemoteBuildingGroupName", "Secondary Remote Build Options");
 	IDetailGroup& RemoteBuildingGroup = BuildCategory.AddGroup(*RemoteBuildingGroupName.ToString(), RemoteBuildingGroupName, false);
 
 	// Remote Server Name Property
@@ -1090,7 +1090,7 @@ void FIOSTargetSettingsCustomization::BuildSecondaryRemoteMacBuildingSection(IDe
 		.FillWidth(1.0f)
 		[
 			SNew(STextBlock)
-			.Text(LOCTEXT("RemoteServerNameLabel", "Secondary Remote Server Name"))
+			.Text(LOCTEXT("SecondaryRemoteServerNameLabel", "Secondary Remote Server Name"))
 		.Font(DetailLayout.GetDetailFont())
 		]
 		]
@@ -1135,7 +1135,7 @@ void FIOSTargetSettingsCustomization::BuildSecondaryRemoteMacBuildingSection(IDe
 		.FillWidth(1.0f)
 		[
 			SNew(STextBlock)
-			.Text(LOCTEXT("RSyncUserNameLabel", "Secondary RSync User Name"))
+			.Text(LOCTEXT("SecondaryRSyncUserNameLabel", "Secondary RSync User Name"))
 		.Font(DetailLayout.GetDetailFont())
 		]
 		]
@@ -1175,7 +1175,7 @@ void FIOSTargetSettingsCustomization::BuildSecondaryRemoteMacBuildingSection(IDe
 	TSharedRef<IPropertyHandle> CwRsyncOverridePathPropertyHandle = DetailLayout.GetProperty(GET_MEMBER_NAME_CHECKED(UIOSRuntimeSettings, SecondaryCwRsyncInstallPath));
 	IDetailPropertyRow& CwRsyncOverridePathPropertyRow = RemoteBuildingGroup.AddPropertyRow(CwRsyncOverridePathPropertyHandle);
 
-	const FText GenerateSSHText = LOCTEXT("GenerateSSHKey", "Secondary Generate SSH Key");
+	const FText GenerateSSHText = LOCTEXT("SecondaryGenerateSSHKey", "Secondary Generate SSH Key");
 
 	// Add a generate key button
 	RemoteBuildingGroup.AddWidgetRow()
