TargetPlatform remove unnecessary TextureCube casts
behavior nop

diff --git a/Engine/Source/Developer/Android/AndroidTargetPlatform/Private/AndroidTargetPlatform.cpp b/Engine/Source/Developer/Android/AndroidTargetPlatform/Private/AndroidTargetPlatform.cpp
--- a/Engine/Source/Developer/Android/AndroidTargetPlatform/Private/AndroidTargetPlatform.cpp
+++ b/Engine/Source/Developer/Android/AndroidTargetPlatform/Private/AndroidTargetPlatform.cpp
@@ -498,19 +498,13 @@ void FAndroidTargetPlatform::GetTextureFormats( const UTexture* Texture, TArray<
 			TextureFormatName = AndroidTexFormat::NameAutoDXT;
 		}
 		
-		// @todo : IsA then CastChecked ?  Just Cast<> ?
-		//	also use ETextureClass instead
-		if (Texture->IsA(UTextureCube::StaticClass()))
+		if (Texture->GetTextureClass() == ETextureClass::Cube) 
 		{
-			const UTextureCube* Cube = CastChecked<UTextureCube>(Texture);
-			if (Cube != nullptr) 
+			FTextureFormatSettings FormatSettings;
+			Texture->GetDefaultFormatSettings(FormatSettings);
+			if (FormatSettings.CompressionSettings == TC_EncodedReflectionCapture && !FormatSettings.CompressionNone)
 			{
-				FTextureFormatSettings FormatSettings;
-				Cube->GetDefaultFormatSettings(FormatSettings);
-				if (FormatSettings.CompressionSettings == TC_EncodedReflectionCapture && !FormatSettings.CompressionNone)
-				{
-					TextureFormatName = FName(TEXT("ETC2_RGBA"));
-				}
+				TextureFormatName = FName(TEXT("ETC2_RGBA"));
 			}
 		}
 
diff --git a/Engine/Source/Developer/IOS/IOSTargetPlatform/Private/IOSTargetPlatform.cpp b/Engine/Source/Developer/IOS/IOSTargetPlatform/Private/IOSTargetPlatform.cpp
--- a/Engine/Source/Developer/IOS/IOSTargetPlatform/Private/IOSTargetPlatform.cpp
+++ b/Engine/Source/Developer/IOS/IOSTargetPlatform/Private/IOSTargetPlatform.cpp
@@ -19,7 +19,6 @@
 #include "Windows/WindowsHWrapper.h"
 #endif
 #if WITH_ENGINE
-#include "Engine/TextureCube.h"
 #include "Sound/SoundWave.h"
 #include "TextureResource.h"
 #include "AudioCompressionSettings.h"
@@ -635,17 +634,13 @@ void FIOSTargetPlatform::GetTextureFormats( const UTexture* Texture, TArray< TAr
 
 	for (FName& TextureFormatName : OutFormats.Last())
 	{
-		if (Texture->IsA(UTextureCube::StaticClass()))
+		if (Texture->GetTextureClass() == ETextureClass::Cube) 
 		{
-			const UTextureCube* Cube = CastChecked<UTextureCube>(Texture);
-			if (Cube != nullptr)
+			FTextureFormatSettings FormatSettings;
+			Texture->GetDefaultFormatSettings(FormatSettings);
+			if (FormatSettings.CompressionSettings == TC_EncodedReflectionCapture && !FormatSettings.CompressionNone)
 			{
-				FTextureFormatSettings FormatSettings;
-				Cube->GetDefaultFormatSettings(FormatSettings);
-				if (FormatSettings.CompressionSettings == TC_EncodedReflectionCapture && !FormatSettings.CompressionNone)
-				{
-					TextureFormatName = FName(TEXT("ETC2_RGBA"));
-				}
+				TextureFormatName = FName(TEXT("ETC2_RGBA"));
 			}
 		}
 	}
