Add a IOS_Minimum version to remove the need from specifying it in ini file.

diff --git a/Engine/Config/BaseEngine.ini b/Engine/Config/BaseEngine.ini
--- a/Engine/Config/BaseEngine.ini
+++ b/Engine/Config/BaseEngine.ini
@@ -2856,7 +2856,6 @@ DistroCodeSigningIdentity=iPhone Distribution:;Apple Distribution:
 VersionInfo=1.0
 FrameRateLock=PUFRL_30
 bEnableDynamicMaxFPS=False
-MinimumiOSVersion=IOS_15
 bSupportsIPad=True
 bSupportsIPhone=True
 AdditionalPlistData=
diff --git a/Engine/Source/Developer/IOS/IOSPlatformEditor/Private/IOSTargetSettingsCustomization.cpp b/Engine/Source/Developer/IOS/IOSPlatformEditor/Private/IOSTargetSettingsCustomization.cpp
--- a/Engine/Source/Developer/IOS/IOSPlatformEditor/Private/IOSTargetSettingsCustomization.cpp
+++ b/Engine/Source/Developer/IOS/IOSPlatformEditor/Private/IOSTargetSettingsCustomization.cpp
@@ -2034,7 +2034,7 @@ void FIOSTargetSettingsCustomization::SetShaderStandard(int32 Value)
 
     if (MinOSPropertyHandle.IsValid())
     {
-        uint8 IOSVersion = (uint8)EIOSVersion::IOS_15;
+        uint8 IOSVersion = (uint8)EIOSVersion::IOS_Minimum;
         if (MinOSPropertyHandle.IsValid())
         {
             MinOSPropertyHandle->GetValue(IOSVersion);
@@ -2062,9 +2062,9 @@ void FIOSTargetSettingsCustomization::UpdateOSVersionWarning()
 		
 		if (bMRTEnabled)
 		{
-			if (EnumValue < (uint8)EIOSVersion::IOS_15)
+			if ((EnumValue < (uint8)EIOSVersion::IOS_15) && (EnumValue != (uint8)EIOSVersion::IOS_Minimum))
 			{
-				SetMinVersion((int32)EIOSVersion::IOS_15);
+				SetMinVersion((int32)EIOSVersion::IOS_Minimum);
 			}
 		}
 		else
@@ -2083,7 +2083,7 @@ void FIOSTargetSettingsCustomization::UpdateOSVersionWarning()
     {
         case (int32)EIOSMetalShaderStandard::IOSMetalSLStandard_Minimum:
         case (int32)EIOSMetalShaderStandard::IOSMetalSLStandard_2_4:
-            if (EnumValue < (uint8)EIOSVersion::IOS_15) {IOSVersionWarningTextBox->SetError(TEXT("iOS15 is the Minimum for Metal 2.4")); return;}
+            if ((EnumValue < (uint8)EIOSVersion::IOS_15) && (EnumValue != (uint8)EIOSVersion::IOS_Minimum)) {IOSVersionWarningTextBox->SetError(TEXT("iOS15 is the Minimum for Metal 2.4")); return;}
             break;
         case (int32)EIOSMetalShaderStandard::IOSMetalSLStandard_3_0:
             if (EnumValue < (uint8)EIOSVersion::IOS_16) {IOSVersionWarningTextBox->SetError(TEXT("iOS16 is the Minimum for Metal 3.0")); return;}
@@ -2104,7 +2104,7 @@ void FIOSTargetSettingsCustomization::UpdateMetalMRTWarning()
 			MinOSPropertyHandle->GetValue(EnumValue);
 			if (EnumValue < (uint8)EIOSVersion::IOS_15)
 			{
-				SetMinVersion((int32)EIOSVersion::IOS_15);
+				SetMinVersion((int32)EIOSVersion::IOS_Minimum);
 			}
 			
 			ShaderVersionPropertyHandle->GetValue(EnumValue);
diff --git a/Engine/Source/Programs/UnrealBuildTool/Platform/IOS/UEBuildIOS.cs b/Engine/Source/Programs/UnrealBuildTool/Platform/IOS/UEBuildIOS.cs
--- a/Engine/Source/Programs/UnrealBuildTool/Platform/IOS/UEBuildIOS.cs
+++ b/Engine/Source/Programs/UnrealBuildTool/Platform/IOS/UEBuildIOS.cs
@@ -417,12 +417,13 @@ namespace UnrealBuildTool
 			{
 				switch (MinimumIOSVersion)
 				{
-					case "IOS_14":
-						return "14.0";
+					case "IOS_Minimum":
 					case "IOS_15":
 						return "15.0";
+					case "IOS_16":
+						return "16.0";
 					default:
-						return "14.0";
+						return "15.0";
 				}
 			}
 		}
diff --git a/Engine/Source/Programs/UnrealBuildTool/Platform/IOS/UEDeployIOS.cs b/Engine/Source/Programs/UnrealBuildTool/Platform/IOS/UEDeployIOS.cs
--- a/Engine/Source/Programs/UnrealBuildTool/Platform/IOS/UEDeployIOS.cs
+++ b/Engine/Source/Programs/UnrealBuildTool/Platform/IOS/UEDeployIOS.cs
@@ -203,15 +203,16 @@ namespace UnrealBuildTool
 			string MinVersionToReturn = "";
 			switch (MinVersion)
 			{
-				case "IOS_14":
-					MinVersionToReturn = "14.0";
-					break;
 				case "IOS_15":
+				case "IOS_Minimum":
 					MinVersionToReturn = "15.0";
 					break;
+				case "IOS_16":
+					MinVersionToReturn = "16.0";
+					break;
 				default:
-					Logger.LogWarning("MinimumiOSVersion {MinVersion} specified in ini file is no longer supported, defaulting to 14.0", MinVersion);
-					MinVersionToReturn = "14.0";
+					MinVersionToReturn = "15.0";
+					Logger.LogWarning("MinimumiOSVersion {MinVersion} specified in ini file is no longer supported, defaulting to {MinVersionToReturn}", MinVersion, MinVersionToReturn);
 					break;
 			}
 			return MinVersionToReturn;
diff --git a/Engine/Source/Runtime/IOS/IOSRuntimeSettings/Classes/IOSRuntimeSettings.h b/Engine/Source/Runtime/IOS/IOSRuntimeSettings/Classes/IOSRuntimeSettings.h
--- a/Engine/Source/Runtime/IOS/IOSRuntimeSettings/Classes/IOSRuntimeSettings.h
+++ b/Engine/Source/Runtime/IOS/IOSRuntimeSettings/Classes/IOSRuntimeSettings.h
@@ -30,10 +30,8 @@ enum class EPowerUsageFrameRateLock : uint8
 UENUM()
 	enum class EIOSVersion : uint8
 {
-    /** iOS 15 */
+    IOS_Minimum = 0 UMETA(DisplayName = "Minimum, Currently 15.0"),
     IOS_15 = 15 UMETA(DisplayName = "15.0"),
-    
-    /** iOS 16 */
     IOS_16 = 16 UMETA(DisplayName = "16.0"),
 };
 
diff --git a/Engine/Source/Runtime/IOS/IOSRuntimeSettings/Private/IOSRuntimeSettings.cpp b/Engine/Source/Runtime/IOS/IOSRuntimeSettings/Private/IOSRuntimeSettings.cpp
--- a/Engine/Source/Runtime/IOS/IOSRuntimeSettings/Private/IOSRuntimeSettings.cpp
+++ b/Engine/Source/Runtime/IOS/IOSRuntimeSettings/Private/IOSRuntimeSettings.cpp
@@ -30,7 +30,7 @@ UIOSRuntimeSettings::UIOSRuntimeSettings(const FObjectInitializer& ObjectInitial
 	bEnableDynamicMaxFPS = false;
 	bSupportsIPad = true;
 	bSupportsIPhone = true;
-	MinimumiOSVersion = EIOSVersion::IOS_15;
+	MinimumiOSVersion = EIOSVersion::IOS_Minimum;
     bBuildAsFramework = true;
 	bGeneratedSYMFile = false;
 	bGeneratedSYMBundle = false;
@@ -127,9 +127,9 @@ void UIOSRuntimeSettings::PostInitProperties()
 		}
 	}
 
-	if (MinimumiOSVersion < EIOSVersion::IOS_15)
+	if ((MinimumiOSVersion < EIOSVersion::IOS_15) && (MinimumiOSVersion != EIOSVersion::IOS_Minimum))
 	{
-		MinimumiOSVersion = EIOSVersion::IOS_15;
+		MinimumiOSVersion = EIOSVersion::IOS_Minimum;
 		UpdateSinglePropertyInConfigFile(GetClass()->FindPropertyByName(GET_MEMBER_NAME_CHECKED(UIOSRuntimeSettings, MinimumiOSVersion)), GetDefaultConfigFilename());
 	}
 	if (!bSupportsMetal && !bSupportsMetalMRT)
