Fix display of multiple iOS/tvOS devices in the QuickLaunch menu within the editor. The iOS icon will now show the type of device connection (wifi/usb)
  - On a mouse hover over the mobile device (includes Android), show additional info about the OS version and model of the device
  - fixes top level device not showing the last selected device
  - don't reshow the selected device in the sub-menu

diff --git a/Engine/Source/Developer/Android/AndroidTargetPlatform/Private/AndroidTargetDevice.h b/Engine/Source/Developer/Android/AndroidTargetPlatform/Private/AndroidTargetDevice.h
--- a/Engine/Source/Developer/Android/AndroidTargetPlatform/Private/AndroidTargetDevice.h
+++ b/Engine/Source/Developer/Android/AndroidTargetPlatform/Private/AndroidTargetDevice.h
@@ -149,6 +149,16 @@ public:
 		return FString::Printf(TEXT("%s (%s)"), *Model, *SerialNumber);
 	}
 
+	virtual FString GetModelId() const override
+	{
+		return Model;
+	}
+
+	virtual FString GetOSVersion() const override
+	{
+		return AndroidVersionString;
+	}
+
 	virtual FString GetOperatingSystemName() override;
 
 	virtual int32 GetProcessSnapshot( TArray<FTargetDeviceProcessInfo>& OutProcessInfos ) override;
diff --git a/Engine/Source/Developer/IOS/IOSTargetPlatform/Private/IOSDeviceHelper.cpp b/Engine/Source/Developer/IOS/IOSTargetPlatform/Private/IOSDeviceHelper.cpp
--- a/Engine/Source/Developer/IOS/IOSTargetPlatform/Private/IOSDeviceHelper.cpp
+++ b/Engine/Source/Developer/IOS/IOSTargetPlatform/Private/IOSDeviceHelper.cpp
@@ -25,22 +25,27 @@ DEFINE_LOG_CATEGORY_STATIC(LogIOSDeviceHelper, Log, All);
 
     struct FDeviceNotificationCallbackInformation
     {
-        FString UDID;
+        FString DeviceID;
         FString DeviceName;
+        FString DeviceUDID;
         FString ProductType;
+        FString DeviceOSVersion;
         DeviceConnectionInterface DeviceInterface;
         uint32 msgType;
-		bool isAuthorized;
+		bool IsAuthorized;
     };
 
 
     struct LibIMobileDevice
     {
-        FString DeviceName;
         FString DeviceID;
+        FString DeviceName;
+        FString DeviceUDID;
         FString DeviceType;
-		bool isAuthorized;
+        FString DeviceOSVersion;
         DeviceConnectionInterface DeviceInterface;
+		bool IsAuthorized;
+		bool IsDealtWith;
     };
 
     static TArray<LibIMobileDevice> GetLibIMobileDevices()
@@ -61,7 +66,7 @@ DEFINE_LOG_CATEGORY_STATIC(LogIOSDeviceHelper, Log, All);
         TArray<FString> DeviceStrings;
         for (int32 StringIndex = 0; StringIndex < OngoingDeviceIds.Num(); ++StringIndex)
         {
-            const FString& DeviceID = OngoingDeviceIds[StringIndex];
+            const FString& DeviceUDID = OngoingDeviceIds[StringIndex];
             DeviceConnectionInterface OngoingDeviceInterface = DeviceConnectionInterface::NoValue;
 
             FString OutStdOutInfo;
@@ -82,11 +87,11 @@ DEFINE_LOG_CATEGORY_STATIC(LogIOSDeviceHelper, Log, All);
 
             if (OngoingDeviceInterface == DeviceConnectionInterface::USB)
             {
-                Arguments = "-u " + DeviceID;
+                Arguments = "-u " + DeviceUDID;
             }
             else if (OngoingDeviceInterface == DeviceConnectionInterface::Network)
             {
-                Arguments = "-n -u " + DeviceID;
+                Arguments = "-n -u " + DeviceUDID;
             }
 
             FPlatformProcess::ExecProcess(*LibimobileDeviceInfo, *Arguments, &ReturnCodeInfo, &OutStdOutInfo, &OutStdErrInfo, NULL, true);
@@ -101,7 +106,7 @@ DEFINE_LOG_CATEGORY_STATIC(LogIOSDeviceHelper, Log, All);
 				{
 					UE_LOG(LogIOSDeviceHelper, Warning, TEXT("Could not pair with connected iOS/tvOS device. Trust this computer by accepting the popup on device."));
 					FString LibimobileDevicePair = GetLibImobileDeviceExe("idevicepair");
-					FString PairArguments = "-u " + DeviceID + " pair";
+					FString PairArguments = "-u " + DeviceUDID + " pair";
 					FPlatformProcess::ExecProcess(*LibimobileDevicePair, *PairArguments, &ReturnCodeInfo, &OutStdOutInfo, &OutStdErrInfo, NULL, true);
 				}
 				else
@@ -110,54 +115,81 @@ DEFINE_LOG_CATEGORY_STATIC(LogIOSDeviceHelper, Log, All);
 				}
 				OutStdOutInfo.Empty();
 				OutStdErrInfo.Empty();
-				ToAdd.isAuthorized = false;
+				ToAdd.IsAuthorized = false;
 			}
 			else
 			{
-				ToAdd.isAuthorized = true;
+				ToAdd.IsAuthorized = true;
 			}
 
 			// parse product type and device name
 			FString DeviceName;
-            OutStdOutInfo.Split(TEXT("DeviceName: "), nullptr, &DeviceName, ESearchCase::CaseSensitive, ESearchDir::FromStart);
+			OutStdOutInfo.Split(TEXT("DeviceName: "), nullptr, &DeviceName, ESearchCase::CaseSensitive, ESearchDir::FromStart);
 			DeviceName.Split(LINE_TERMINATOR, &DeviceName, nullptr, ESearchCase::CaseSensitive, ESearchDir::FromStart);
-			if (!ToAdd.isAuthorized)
-            {
+			if (!ToAdd.IsAuthorized)
+			{
 				DeviceName = LOCTEXT("IosTvosUnauthorizedDevice", "iOS / tvOS (Unauthorized)").ToString();
-	        }
-            FString ProductType;
-            OutStdOutInfo.Split(TEXT("ProductType: "), nullptr, &ProductType, ESearchCase::CaseSensitive, ESearchDir::FromStart);
-            ProductType.Split(LINE_TERMINATOR, &ProductType, nullptr, ESearchCase::CaseSensitive, ESearchDir::FromStart);
-            ToAdd.DeviceID = DeviceID;
-            ToAdd.DeviceName = DeviceName;
-            ToAdd.DeviceType = ProductType;
-            ToAdd.DeviceInterface = OngoingDeviceInterface;
-            ToReturn.Add(ToAdd);
-        }
-        return ToReturn;
-    }
+			}
+			else
+			{
+				if (OngoingDeviceInterface == DeviceConnectionInterface::Network)
+				{
+					DeviceName += " [Wifi]";
+				}
+			}
+			
+			FString ProductType;
+			OutStdOutInfo.Split(TEXT("ProductType: "), nullptr, &ProductType, ESearchCase::CaseSensitive, ESearchDir::FromStart);
+			ProductType.Split(LINE_TERMINATOR, &ProductType, nullptr, ESearchCase::CaseSensitive, ESearchDir::FromStart);
+			
+			FString OSVersion; // iOS/iPad OS Version
+			OutStdOutInfo.Split(TEXT("ProductVersion: "), nullptr, &OSVersion, ESearchCase::CaseSensitive, ESearchDir::FromStart);
+			OSVersion.Split(LINE_TERMINATOR, &OSVersion, nullptr, ESearchCase::CaseSensitive, ESearchDir::FromStart);
+			
+			FString DeviceID = FString::Printf(TEXT("%s@%s"),
+											   ProductType.Contains(TEXT("AppleTV")) ? TEXT("TVOS") : TEXT("IOS"),
+											   *DeviceUDID);
+			
+			ToAdd.DeviceID = DeviceID;
+			ToAdd.DeviceUDID = DeviceUDID;
+			ToAdd.DeviceName = DeviceName;
+			ToAdd.DeviceType = ProductType;
+			ToAdd.DeviceOSVersion = OSVersion;
+			ToAdd.DeviceInterface = OngoingDeviceInterface;
+			ToAdd.IsDealtWith = false;
+			ToReturn.Add(ToAdd);
+		}
+		return ToReturn;
+	}
 
 class FIOSDevice
 {
 public:
-    FIOSDevice(FString InID, FString InName)
-        : UDID(InID)
-        , Name(InName)
-    {
-    }
-    
-    ~FIOSDevice()
-    {
-    }
-
-    FString SerialNumber() const
-    {
-        return UDID;
-    }
-
+	FIOSDevice(FString InID, FString InName, DeviceConnectionInterface InConnectionType)
+		: UDID(InID)
+		, Name(InName)
+		, ConnectionType(InConnectionType)
+	{
+	}
+	
+	~FIOSDevice()
+	{
+	}
+	
+	FString SerialNumber() const
+	{
+		return UDID;
+	}
+	
+	DeviceConnectionInterface ConnectionInterface() const
+	{
+		return ConnectionType;
+	}
+	
 private:
-    FString UDID;
-    FString Name;
+	FString UDID;
+	FString Name;
+	DeviceConnectionInterface ConnectionType;
 };
 
 /**
@@ -284,16 +316,20 @@ private:
 
         if (bAdd)
         {
+            CallbackInfo.DeviceID = Device.DeviceID;
             CallbackInfo.DeviceName = Device.DeviceName;
-            CallbackInfo.UDID = Device.DeviceID;
+            CallbackInfo.DeviceUDID = Device.DeviceUDID;
             CallbackInfo.DeviceInterface = Device.DeviceInterface;
             CallbackInfo.ProductType = Device.DeviceType;
+            CallbackInfo.DeviceOSVersion = Device.DeviceOSVersion;
             CallbackInfo.msgType = 1;
-			CallbackInfo.isAuthorized = Device.isAuthorized;
+			CallbackInfo.IsAuthorized = Device.IsAuthorized;
         }
         else
         {
-            CallbackInfo.UDID = Device.DeviceID;
+            CallbackInfo.DeviceID = Device.DeviceID;
+			CallbackInfo.DeviceUDID = Device.DeviceUDID;
+			CallbackInfo.DeviceInterface = Device.DeviceInterface;
             CallbackInfo.msgType = 2;
             DeviceNotification.Broadcast(&CallbackInfo);
 
@@ -339,16 +375,17 @@ private:
         {
 			LibIMobileDevice* Found = CachedDevices.FindByPredicate(
 				[&](LibIMobileDevice Element) {
-					return Element.DeviceID == ParsedDevices[Index].DeviceID;
+					return (Element.DeviceUDID == ParsedDevices[Index].DeviceUDID &&
+						Element.DeviceInterface == ParsedDevices[Index].DeviceInterface);
 				});
 			if (Found != nullptr)
 			{
-				if (Found->isAuthorized != ParsedDevices[Index].isAuthorized)
+				if (Found->IsAuthorized != ParsedDevices[Index].IsAuthorized)
 				{
 					NotifyDeviceChange(ParsedDevices[Index], false);
 					NotifyDeviceChange(ParsedDevices[Index], true);
 				}
-				Found->DeviceID = "DealtWith";
+				Found->IsDealtWith = true;
 			}
 			else
 			{
@@ -359,7 +396,7 @@ private:
 
         for (int32 Index = 0; Index < CachedDevices.Num(); ++Index)
         {
-            if (CachedDevices[Index].DeviceID != "DealtWith")
+            if (!CachedDevices[Index].IsDealtWith)
             {
                 NotifyDeviceChange(CachedDevices[Index], false);
             }
@@ -448,14 +485,18 @@ void FIOSDeviceHelper::DoDeviceConnect(void* CallbackInfo)
 {
     // connect to the device
     struct FDeviceNotificationCallbackInformation* cbi = (FDeviceNotificationCallbackInformation*)CallbackInfo;
-    FIOSDevice* Device = new FIOSDevice(cbi->UDID, cbi->DeviceName);
+    FIOSDevice* Device = new FIOSDevice(cbi->DeviceUDID, cbi->DeviceName, cbi->DeviceInterface);
 
     // fire the event
     FIOSLaunchDaemonPong Event;
-    Event.DeviceID = FString::Printf(TEXT("%s@%s"), cbi->ProductType.Contains(TEXT("AppleTV")) ? TEXT("TVOS") : TEXT("IOS"), *(cbi->UDID));
+	Event.DeviceID = cbi->DeviceID;
+    Event.DeviceUDID = cbi->DeviceUDID;
     Event.DeviceName = cbi->DeviceName;
     Event.DeviceType = cbi->ProductType;
-	Event.bIsAuthorized = cbi->isAuthorized;
+    Event.DeviceOSVersion = cbi->DeviceOSVersion;
+    Event.DeviceModelId = cbi->ProductType;
+    Event.DeviceConnectionType = (cbi->DeviceInterface == DeviceConnectionInterface::Network) ? "Network" : "USB";
+	Event.bIsAuthorized = cbi->IsAuthorized;
     Event.bCanReboot = false;
     Event.bCanPowerOn = false;
     Event.bCanPowerOff = false;
@@ -467,27 +508,29 @@ void FIOSDeviceHelper::DoDeviceConnect(void* CallbackInfo)
 
 void FIOSDeviceHelper::DoDeviceDisconnect(void* CallbackInfo)
 {
-    struct FDeviceNotificationCallbackInformation* cbi = (FDeviceNotificationCallbackInformation*)CallbackInfo;
-    FIOSDevice* device = NULL;
-    for (auto DeviceIterator = ConnectedDevices.CreateIterator(); DeviceIterator; ++DeviceIterator)
-    {
-        if (DeviceIterator.Key()->SerialNumber() == cbi->UDID)
-        {
-            device = DeviceIterator.Key();
-            break;
-        }
-    }
-    if (device != NULL)
-    {
-        // extract the device id from the connected list
-        FIOSLaunchDaemonPong Event = ConnectedDevices.FindAndRemoveChecked(device);
-
-        // fire the event
-        FIOSDeviceHelper::OnDeviceDisconnected().Broadcast(Event);
-
-        // delete the device
-        delete device;
-    }
+	struct FDeviceNotificationCallbackInformation* cbi = (FDeviceNotificationCallbackInformation*)CallbackInfo;
+	FIOSDevice* device = NULL;
+	for (auto DeviceIterator = ConnectedDevices.CreateIterator(); DeviceIterator; ++DeviceIterator)
+	{
+		if (DeviceIterator.Key()->SerialNumber() == cbi->DeviceUDID &&
+			DeviceIterator.Key()->ConnectionInterface() == cbi->DeviceInterface)
+		{
+			device = DeviceIterator.Key();
+			break;
+		}
+	}
+	
+	if (device != NULL)
+	{
+		// extract the device id from the connected list
+		FIOSLaunchDaemonPong Event = ConnectedDevices.FindAndRemoveChecked(device);
+		
+		// fire the event
+		FIOSDeviceHelper::OnDeviceDisconnected().Broadcast(Event);
+		
+		// delete the device
+		delete device;
+	}
 }
 
 bool FIOSDeviceHelper::InstallIPAOnDevice(const FTargetDeviceId& DeviceId, const FString& IPAPath)
diff --git a/Engine/Source/Developer/IOS/IOSTargetPlatform/Private/IOSTargetDevice.cpp b/Engine/Source/Developer/IOS/IOSTargetPlatform/Private/IOSTargetDevice.cpp
--- a/Engine/Source/Developer/IOS/IOSTargetPlatform/Private/IOSTargetDevice.cpp
+++ b/Engine/Source/Developer/IOS/IOSTargetPlatform/Private/IOSTargetDevice.cpp
@@ -22,10 +22,12 @@ FIOSTargetDevice::FIOSTargetDevice(const ITargetPlatform& InTargetPlatform)
 	, bCanPowerOn(false)
 	, bCanPowerOff(false)
 	, DeviceType(ETargetDeviceTypes::Indeterminate)
+	, DeviceModelId(TEXT(""))
 {
 	DeviceId = FTargetDeviceId(TargetPlatform.PlatformName(), FPlatformProcess::ComputerName());
 	DeviceName = FPlatformProcess::ComputerName();
 	MessageEndpoint = FMessageEndpoint::Builder("FIOSTargetDevice").Build();
+	DeviceConnectionType = ETargetDeviceConnectionTypes::USB;
 }
 
 
@@ -51,6 +53,11 @@ ETargetDeviceTypes FIOSTargetDevice::GetDeviceType() const
 	return DeviceType;
 }
 
+ETargetDeviceConnectionTypes FIOSTargetDevice::GetDeviceConnectionType() const
+{
+	return DeviceConnectionType;
+}
+
 FTargetDeviceId FIOSTargetDevice::GetId() const
 {
 	return DeviceId;
@@ -66,6 +73,16 @@ FString FIOSTargetDevice::GetOperatingSystemName()
 	return TargetPlatform.PlatformName();
 }
 
+FString FIOSTargetDevice::GetModelId() const
+{
+	return DeviceModelId;
+}
+
+FString FIOSTargetDevice::GetOSVersion() const
+{
+	return DeviceOSVersion;
+}
+
 const class ITargetPlatform& FIOSTargetDevice::GetTargetPlatform() const
 {
 	return TargetPlatform;
diff --git a/Engine/Source/Developer/IOS/IOSTargetPlatform/Private/IOSTargetDevice.h b/Engine/Source/Developer/IOS/IOSTargetPlatform/Private/IOSTargetDevice.h
--- a/Engine/Source/Developer/IOS/IOSTargetPlatform/Private/IOSTargetDevice.h
+++ b/Engine/Source/Developer/IOS/IOSTargetPlatform/Private/IOSTargetDevice.h
@@ -112,9 +112,12 @@ public:
 	virtual void Disconnect() override;
 	virtual int32 GetProcessSnapshot(TArray<FTargetDeviceProcessInfo>& OutProcessInfos) override;
 	virtual ETargetDeviceTypes GetDeviceType() const override;
+	virtual ETargetDeviceConnectionTypes GetDeviceConnectionType() const override;
 	virtual FTargetDeviceId GetId() const override;
 	virtual FString GetName() const override;
 	virtual FString GetOperatingSystemName() override;
+	virtual FString GetModelId() const override;
+	virtual FString GetOSVersion() const override;
 	virtual const class ITargetPlatform& GetTargetPlatform() const override;
 	virtual bool IsConnected() override;
 	virtual bool IsDefault() const override;
@@ -180,6 +183,15 @@ private:
 	/** Type of device */
 	ETargetDeviceTypes DeviceType;
 
+	/** The specific model identifier of the device */
+	FString DeviceModelId;
+
+	/** The iOS/tvOS/iPadOS OS version */
+	FString DeviceOSVersion;
+
+	/** Type of device connection (USB or Wifi) */
+	ETargetDeviceConnectionTypes DeviceConnectionType;
+
 public:
 
 	void SetFeature(ETargetDeviceFeatures InFeature, bool bFlag)
@@ -210,11 +222,23 @@ public:
 		DeviceName = InDeviceName;
 	}
 
+	/** Sets the modelId of the device */
+	void SetModelId(const FString InModelId)
+	{
+		DeviceModelId = InModelId;
+	}
+
+	/** Sets the OS version of the device */
+	void SetOSVersion(const FString InOSVersion)
+	{
+		DeviceOSVersion = InOSVersion;
+	}
+
 	/**
- * Sets the device's authorization state.
- *
- * @param bInIsAuthorized - Whether the device is authorized for USB communications.
- */
+	 * Sets the device's authorization state.
+	 *
+	 * @param bInIsAuthorized - Whether the device is authorized for USB communications.
+	 */
 	void SetAuthorized(bool bInIsAuthorized)
 	{
 		bIsDeviceAuthorized = bInIsAuthorized;
@@ -245,6 +269,23 @@ public:
 		}
 	}
 
+	/** Sets the connection type (usb/wifi) of the device */
+	void SetDeviceConnectionType(const FString InDeviceConnectionTypeString)
+	{
+		if (InDeviceConnectionTypeString == TEXT("Network"))
+		{
+			DeviceConnectionType = ETargetDeviceConnectionTypes::Wifi;
+		}
+		else if (InDeviceConnectionTypeString == TEXT("USB"))
+		{
+			DeviceConnectionType = ETargetDeviceConnectionTypes::USB;
+		}
+		else
+		{
+			DeviceConnectionType = ETargetDeviceConnectionTypes::Unknown;
+		}
+	}
+
 	void SetDeviceEndpoint(const FMessageAddress& DeviceAddress)
 	{
 		DeviceEndpoint = DeviceAddress;
diff --git a/Engine/Source/Developer/IOS/IOSTargetPlatform/Private/IOSTargetPlatform.cpp b/Engine/Source/Developer/IOS/IOSTargetPlatform/Private/IOSTargetPlatform.cpp
--- a/Engine/Source/Developer/IOS/IOSTargetPlatform/Private/IOSTargetPlatform.cpp
+++ b/Engine/Source/Developer/IOS/IOSTargetPlatform/Private/IOSTargetPlatform.cpp
@@ -366,6 +366,9 @@ void FIOSTargetPlatform::HandlePongMessage( const FIOSLaunchDaemonPong& Message,
 		Device->SetDeviceId(DeviceId);
 		Device->SetDeviceName(Message.DeviceName);
 		Device->SetDeviceType(Message.DeviceType);
+		Device->SetModelId(Message.DeviceModelId);
+		Device->SetOSVersion(Message.DeviceOSVersion);
+		Device->SetDeviceConnectionType(Message.DeviceConnectionType);
 		Device->SetDeviceEndpoint(Context->GetSender());
 		Device->SetIsSimulated(Message.DeviceID.Contains(TEXT("Simulator")));
 
@@ -395,6 +398,9 @@ void FIOSTargetPlatform::HandleDeviceConnected(const FIOSLaunchDaemonPong& Messa
 			Device->SetDeviceName(Message.DeviceName);
 			Device->SetAuthorized(Message.bIsAuthorized);
 			Device->SetDeviceType(Message.DeviceType);
+			Device->SetModelId(Message.DeviceModelId);
+			Device->SetOSVersion(Message.DeviceOSVersion);
+			Device->SetDeviceConnectionType(Message.DeviceConnectionType);
 			Device->SetIsSimulated(Message.DeviceID.Contains(TEXT("Simulator")));
 
 			OnDeviceDiscovered().Broadcast(Device.ToSharedRef());
@@ -493,6 +499,10 @@ bool FIOSTargetPlatform::SupportsFeature( ETargetPlatformFeatures Feature ) cons
 			const bool bUsesARMCompressor = (CompressorCVar ? (CompressorCVar->GetInt() != 0) : false);
 			return bUsesARMCompressor;
 		}
+
+		case ETargetPlatformFeatures::ShowAsPlatformGroup:
+			return false;
+
 		default:
 			break;
 	}
diff --git a/Engine/Source/Developer/TargetDeviceServices/Private/Proxies/TargetDeviceProxy.cpp b/Engine/Source/Developer/TargetDeviceServices/Private/Proxies/TargetDeviceProxy.cpp
--- a/Engine/Source/Developer/TargetDeviceServices/Private/Proxies/TargetDeviceProxy.cpp
+++ b/Engine/Source/Developer/TargetDeviceServices/Private/Proxies/TargetDeviceProxy.cpp
@@ -63,6 +63,8 @@ void FTargetDeviceProxy::UpdateFromMessage( const FTargetDeviceServicePong& Mess
 	Name = Aggregated? Message.AllDevicesName: Message.Name;
 	DeviceUser = Message.DeviceUser;
 	DeviceUserPassword = Message.DeviceUserPassword;
+	OSVersion = Message.OSVersion;
+	ConnectionType = Message.ConnectionType;
 	Shared = Message.Shared;
 	SupportsMultiLaunch = Message.SupportsMultiLaunch;
 	SupportsPowerOff = Message.SupportsPowerOff;
diff --git a/Engine/Source/Developer/TargetDeviceServices/Private/Proxies/TargetDeviceProxy.h b/Engine/Source/Developer/TargetDeviceServices/Private/Proxies/TargetDeviceProxy.h
--- a/Engine/Source/Developer/TargetDeviceServices/Private/Proxies/TargetDeviceProxy.h
+++ b/Engine/Source/Developer/TargetDeviceServices/Private/Proxies/TargetDeviceProxy.h
@@ -156,6 +156,16 @@ public:
 		return Type;
 	}
 
+	virtual const FString& GetOSVersion() const override
+	{
+		return OSVersion;
+	}
+
+	virtual const FString& GetConnectionType() const override
+	{
+		return ConnectionType;
+	}
+
 	virtual bool HasDeviceId(const FString& InDeviceId) const override;
 	virtual bool HasVariant(FName InVariant) const override;
 	virtual bool HasTargetPlatform(FName InTargetPlatformId) const override;
@@ -251,6 +261,12 @@ private:
 	/** Holds the device type. */
 	FString Type;
 
+	/** Holds the remote device's OS version. */
+	FString OSVersion;
+
+	/** Holds the type of connection. */
+	FString ConnectionType;
+
 	/** Holds default variant name. */
 	FName DefaultVariant;
 
diff --git a/Engine/Source/Developer/TargetDeviceServices/Private/Services/TargetDeviceService.cpp b/Engine/Source/Developer/TargetDeviceServices/Private/Services/TargetDeviceService.cpp
--- a/Engine/Source/Developer/TargetDeviceServices/Private/Services/TargetDeviceService.cpp
+++ b/Engine/Source/Developer/TargetDeviceServices/Private/Services/TargetDeviceService.cpp
@@ -425,9 +425,11 @@ void FTargetDeviceService::HandlePingMessage(const FTargetDeviceServicePing& InM
 		Message->HostName = FPlatformProcess::ComputerName();
 		Message->HostUser = FPlatformProcess::UserName(false);
 		Message->Connected = DefaultDevice->IsConnected();
+		Message->ConnectionType = (DefaultDevice->GetDeviceConnectionType() == ETargetDeviceConnectionTypes::Wifi) ? "Network" : "USB";
 		Message->Authorized = DefaultDevice->IsAuthorized();
 		Message->Make = TEXT("@todo");
-		Message->Model = TEXT("@todo");
+		Message->Model = DefaultDevice->GetModelId();
+		Message->OSVersion = DefaultDevice->GetOSVersion();
 		DefaultDevice->GetUserCredentials(Message->DeviceUser, Message->DeviceUserPassword);
 		Message->Shared = Shared;
 		Message->SupportsMultiLaunch = DefaultDevice->SupportsFeature(ETargetDeviceFeatures::MultiLaunch);
diff --git a/Engine/Source/Developer/TargetDeviceServices/Public/ITargetDeviceProxy.h b/Engine/Source/Developer/TargetDeviceServices/Public/ITargetDeviceProxy.h
--- a/Engine/Source/Developer/TargetDeviceServices/Public/ITargetDeviceProxy.h
+++ b/Engine/Source/Developer/TargetDeviceServices/Public/ITargetDeviceProxy.h
@@ -206,6 +206,20 @@ public:
 	 */
 	virtual const FString& GetType() const = 0;
 
+	/**
+	 * Gets the OS Version of the device
+	 *
+	 * @return Device OS VErsion.
+	 */
+	virtual const FString& GetOSVersion() const = 0;
+
+	/**
+	 * Gets the type of connection for this device.
+	 *
+	 * @return Device type.
+	 */
+	virtual const FString& GetConnectionType() const = 0;
+
 	/**
 	 * Checks whether the device is currently connected.
 	 *
diff --git a/Engine/Source/Developer/TargetDeviceServices/Public/TargetDeviceServiceMessages.h b/Engine/Source/Developer/TargetDeviceServices/Public/TargetDeviceServiceMessages.h
--- a/Engine/Source/Developer/TargetDeviceServices/Public/TargetDeviceServiceMessages.h
+++ b/Engine/Source/Developer/TargetDeviceServices/Public/TargetDeviceServiceMessages.h
@@ -288,6 +288,14 @@ struct FTargetDeviceServicePong
 	UPROPERTY(EditAnywhere, Category="Message")
 	FString Type;
 
+	/** Holds the device OS Version. */
+	UPROPERTY(EditAnywhere, Category="Message")
+	FString OSVersion;
+
+	/** Holds the connection type. */
+	UPROPERTY(EditAnywhere, Category="Message")
+	FString ConnectionType;
+
 	/** Holds the variant name of the default device. */
 	UPROPERTY(EditAnywhere, Category="Message")
 	FName DefaultVariant;
diff --git a/Engine/Source/Developer/TargetPlatform/Public/Interfaces/ITargetDevice.h b/Engine/Source/Developer/TargetPlatform/Public/Interfaces/ITargetDevice.h
--- a/Engine/Source/Developer/TargetPlatform/Public/Interfaces/ITargetDevice.h
+++ b/Engine/Source/Developer/TargetPlatform/Public/Interfaces/ITargetDevice.h
@@ -58,6 +58,29 @@ enum class ETargetDeviceTypes
 	HMD
 };
 
+/**
+ * Enumerates how the target device is connected
+ */
+enum class ETargetDeviceConnectionTypes
+{
+	/** It's unknown how the device is connected. */
+	Unknown,
+
+	/** The device is connected through USB. */
+	USB,
+
+	/** The device is connected via Wifi. */
+	Wifi,
+	
+	/** The device is connected via Ethernet. */
+	Ethernet,
+
+	/** The device is running as a simulator on the current Editor. */
+	Simulator,
+
+	/** The device is connected via a proprietary connection. */
+	Proprietary
+};
 
 namespace TargetDeviceTypes
 {
@@ -228,6 +251,27 @@ public:
 	 */
 	virtual ETargetDeviceTypes GetDeviceType() const = 0;
 
+	/**
+	 * Gets the device model identifier.
+	 *
+	 * @return ModelId.
+	 */
+	virtual FString GetModelId() const {return FString();};
+
+	/**
+	 * Gets the device OS Version.
+	 *
+	 * @return OSVersion.
+	 */
+	virtual FString GetOSVersion() const {return FString();};
+
+	/**
+	 * Gets the device connection type.
+	 * 
+	 * @return Device connection type.
+	 */
+	virtual ETargetDeviceConnectionTypes GetDeviceConnectionType() const {return ETargetDeviceConnectionTypes::Unknown; };
+
 	/**
 	 * Gets the unique device identifier.
 	 *
diff --git a/Engine/Source/Developer/TargetPlatform/Public/Interfaces/ITargetPlatform.h b/Engine/Source/Developer/TargetPlatform/Public/Interfaces/ITargetPlatform.h
--- a/Engine/Source/Developer/TargetPlatform/Public/Interfaces/ITargetPlatform.h
+++ b/Engine/Source/Developer/TargetPlatform/Public/Interfaces/ITargetPlatform.h
@@ -112,7 +112,10 @@ enum class ETargetPlatformFeatures
 	DirectDataExchange,
 
 	/** The platform supports Luminance + Alpha encoding mode for normalmaps */
-	NormalmapLAEncodingMode
+	NormalmapLAEncodingMode,
+
+	/** All devices of this platform should be grouped under one platform group */
+	ShowAsPlatformGroup
 };
 
 enum class EPlatformAuthentication
diff --git a/Engine/Source/Editor/EditorStyle/Private/StarshipStyle.cpp b/Engine/Source/Editor/EditorStyle/Private/StarshipStyle.cpp
--- a/Engine/Source/Editor/EditorStyle/Private/StarshipStyle.cpp
+++ b/Engine/Source/Editor/EditorStyle/Private/StarshipStyle.cpp
@@ -7280,7 +7280,10 @@ void FStarshipEditorStyle::FStyle::SetupAutomationStyles()
 			.SetHoveredForeground(FLinearColor::FromSRGBColor(FColor::FromHex("#868686FF")))
 			.SetPressedForeground(FLinearColor::FromSRGBColor(FColor::FromHex("#868686FF")));
 		
-		Set("DeviceDetails.EditButton", DeviceProfileCellButton);
+		Set( "DeviceDetails.EditButton", DeviceProfileCellButton);
+
+		Set( "DeviceDetails.WIFI.IOS", new IMAGE_BRUSH("Starship/DeviceManager/AppleWifi_128x", Icon128x128));
+		Set( "DeviceDetails.USB.IOS", new IMAGE_BRUSH("Starship/DeviceManager/AppleUsb_128x", Icon128x128));
 	}
 
 	// Settings Editor
diff --git a/Engine/Source/Editor/TurnkeySupport/Private/TurnkeySupportModule.cpp b/Engine/Source/Editor/TurnkeySupport/Private/TurnkeySupportModule.cpp
--- a/Engine/Source/Editor/TurnkeySupport/Private/TurnkeySupportModule.cpp
+++ b/Engine/Source/Editor/TurnkeySupport/Private/TurnkeySupportModule.cpp
@@ -933,6 +933,8 @@ static void TurnkeyInstallSdk(FString IniPlatformName, bool bPreferFull, bool bF
 	);
 }
 
+static TMap<FName, FString> PerPlatformLastChosen;
+
 static TAttribute<FText> MakeSdkStatusAttribute(FName IniPlatformName, TSharedPtr< ITargetDeviceProxy> DeviceProxy)
 {
 	FString DisplayString = DeviceProxy ? DeviceProxy->GetName() : IniPlatformName.ToString();
@@ -953,40 +955,59 @@ static TAttribute<FText> MakeSdkStatusAttribute(FName IniPlatformName, TSharedPt
 	}));
 }
 
-//static TAttribute<FSlateIcon> MakePlatformSdkIconAttribute(FName IniPlatformName, TSharedPtr< ITargetDeviceProxy> DeviceProxy)
 static FSlateIcon MakePlatformSdkIconAttribute(FName IniPlatformName, TSharedPtr< ITargetDeviceProxy> DeviceProxy)
 {
 	FString DeviceId = DeviceProxy ? DeviceProxy->GetTargetDeviceId(NAME_None) : FString();
 
-//	return TAttribute<FSlateIcon>::Create(TAttribute<FSlateIcon>::FGetter::CreateLambda([IniPlatformName, DeviceId]()
-		{
-			// get the status, or Unknown if it's not there
-			ETurnkeyPlatformSdkStatus Status = DeviceId.Len() ? ITurnkeySupportModule::Get().GetSdkInfoForDeviceId(DeviceId).Status : ITurnkeySupportModule::Get().GetSdkInfo(IniPlatformName, false).Status;
-			bool bDeviceWarning = (DeviceId.Len() && ITurnkeySupportModule::Get().GetSdkInfoForDeviceId(DeviceId).DeviceStatus != ETurnkeyDeviceStatus::SoftwareValid);
+	// get the status, or Unknown if it's not there
+	ETurnkeyPlatformSdkStatus Status = DeviceId.Len() ? ITurnkeySupportModule::Get().GetSdkInfoForDeviceId(DeviceId).Status : ITurnkeySupportModule::Get().GetSdkInfo(IniPlatformName, false).Status;
+	bool bDeviceWarning = (DeviceId.Len() && ITurnkeySupportModule::Get().GetSdkInfoForDeviceId(DeviceId).DeviceStatus != ETurnkeyDeviceStatus::SoftwareValid);
 
-			if (Status == ETurnkeyPlatformSdkStatus::OutOfDate || Status == ETurnkeyPlatformSdkStatus::NoSdk || bDeviceWarning)
+	if (Status == ETurnkeyPlatformSdkStatus::OutOfDate || Status == ETurnkeyPlatformSdkStatus::NoSdk || bDeviceWarning)
+	{
+		return FSlateIcon(FAppStyle::Get().GetStyleSetName(), TEXT("Icons.Warning"));
+	}
+	else if (Status == ETurnkeyPlatformSdkStatus::Error)
+	{
+		return FSlateIcon(FAppStyle::Get().GetStyleSetName(), TEXT("Icons.Error"));
+	}
+	else if (Status == ETurnkeyPlatformSdkStatus::Unknown)
+	{
+		return FSlateIcon(FAppStyle::Get().GetStyleSetName(), TEXT("Icons.Help"));
+	}
+	else
+	{
+		if (DeviceProxy != nullptr && !DeviceProxy->GetConnectionType().IsEmpty())
+		{
+			FSlateIcon Icon;
+			if (DeviceProxy->GetConnectionType().Contains(TEXT("Network")))
 			{
-				return FSlateIcon(FAppStyle::Get().GetStyleSetName(), TEXT("Icons.Warning"));
+				Icon = FSlateIcon(FAppStyle::Get().GetStyleSetName(), *(TEXT("DeviceDetails.WIFI.") + IniPlatformName.ToString()));
 			}
-			else if (Status == ETurnkeyPlatformSdkStatus::Error)
+			else if (DeviceProxy->GetConnectionType().Contains(TEXT("USB")))
 			{
-				return FSlateIcon(FAppStyle::Get().GetStyleSetName(), TEXT("Icons.Error"));
+				Icon = FSlateIcon(FAppStyle::Get().GetStyleSetName(), *(TEXT("DeviceDetails.USB.") + IniPlatformName.ToString()));
 			}
-			else if (Status == ETurnkeyPlatformSdkStatus::Unknown)
+			else
 			{
-				return FSlateIcon(FAppStyle::Get().GetStyleSetName(), TEXT("Icons.Help"));
+				UE_LOG(LogTurnkeySupport, Warning, TEXT("Unknown ConnectionType:%s"), *DeviceProxy->GetConnectionType());
 			}
-			else
+
+			// Verify that the icon exists by checking it's not using the "DefaultBrush" icon resource
+			const FSlateBrush *brush = FAppStyle::GetBrush(Icon.GetStyleName());
+			if (brush != nullptr && (brush->GetResourceName() != FAppStyle::GetDefaultBrush()->GetResourceName()))
 			{
-				return FSlateIcon(FAppStyle::Get().GetStyleSetName(), FDataDrivenPlatformInfoRegistry::GetPlatformInfo(IniPlatformName).GetIconStyleName(EPlatformIconSize::Normal));
+				return Icon;
 			}
 		}
-//		));
+
+		// Default to platform icon
+		return FSlateIcon(FAppStyle::Get().GetStyleSetName(), FDataDrivenPlatformInfoRegistry::GetPlatformInfo(IniPlatformName).GetIconStyleName(EPlatformIconSize::Normal));
+	}
 }
 
 static void FormatSdkInfo(const FString& PlatformOrDevice, const FTurnkeySdkInfo& SdkInfo, FText& OutInfo, FText& OutToolTip)
 {
-
 	TArray<FText> Lines;
 
 	for (TPair<FString, FTurnkeySdkInfo::Version> Pair : SdkInfo.SDKVersions)
@@ -1667,6 +1688,7 @@ static void GenerateDeviceProxyMenuParams(TSharedPtr<ITargetDeviceProxy> DeviceP
 					DeviceVariantName = *VariantName;
 				}
 
+				PerPlatformLastChosen[PlatformName] = DeviceProxy->GetName();
 				FString DeviceId = DeviceProxy->GetTargetDeviceId(DeviceVariantName);
 				HandleLaunchOnDeviceActionExecute(DeviceId, DeviceProxy->GetName(), true);
 				ExternalOnClickDelegate.ExecuteIfBound(DeviceId);
@@ -1678,11 +1700,23 @@ static void GenerateDeviceProxyMenuParams(TSharedPtr<ITargetDeviceProxy> DeviceP
 	FFormatNamedArguments TooltipArguments;
 	TooltipArguments.Add(TEXT("DeviceID"), FText::FromString(DeviceProxy->GetName()));
 	TooltipArguments.Add(TEXT("DisplayName"), FText::FromName(PlatformName));
-	OutTooltip = FText::Format(LOCTEXT("LaunchDeviceToolTipText_ThisDevice", "Launch the game on this {DisplayName} device ({DeviceID})"), TooltipArguments);
+	
 	if (!DeviceProxy->IsAuthorized())
 	{
 		OutTooltip = FText::Format(LOCTEXT("LaunchDeviceToolTipText_UnauthorizedOrLocked", "{DisplayName} device ({DeviceID}) is unauthorized or locked"), TooltipArguments);
 	}
+	else if (!DeviceProxy->GetModel().IsEmpty() && !DeviceProxy->GetOSVersion().IsEmpty())
+	{
+		// Some platforms (ie: iOS/Android) can display additional information on mouse hover
+		TooltipArguments.Add(TEXT("ModelId"), FText::FromString(DeviceProxy->GetModel()));
+		TooltipArguments.Add(TEXT("OSVersion"), FText::FromString(DeviceProxy->GetOSVersion()));
+
+		OutTooltip = FText::Format(LOCTEXT("LaunchDeviceToolTipText_ThisDeviceExtra", "Launch on \"{DeviceID}\" [OSVer:'{OSVersion}' Model:'{ModelId}']"), TooltipArguments);
+	}
+	else
+	{
+		OutTooltip = FText::Format(LOCTEXT("LaunchDeviceToolTipText_ThisDevice", "Launch the game on this {DisplayName} device ({DeviceID})"), TooltipArguments);
+	}
 
 	FProjectStatus ProjectStatus;
 	if (IProjectManager::Get().QueryStatusForCurrentProject(ProjectStatus) && !ProjectStatus.IsTargetPlatformSupported(PlatformName))
@@ -1721,16 +1755,28 @@ void FTurnkeySupportModule::MakeQuickLaunchItems(class UToolMenu* Menu, FOnQuick
 
 			if (DeviceProxies.Num() > 0)
 			{
-				// 					Algo::Sort(DeviceProxies, [](TSharedPtr<ITargetDeviceProxy> A, TSharedPtr<ITargetDeviceProxy> B)
-				// 					{ return A->Name == LastChosen || B->Name != LastChosen && A->IsDefault(); });
-
+				FString LastChosen;
+				if (PerPlatformLastChosen.Contains(PlatformName))
+				{
+					LastChosen = PerPlatformLastChosen[PlatformName];
+				}
+				else
+				{
+					LastChosen = DeviceProxies[0]->GetName();
+					PerPlatformLastChosen.Add(PlatformName, LastChosen);
+				}
+				Algo::Sort(DeviceProxies, [&LastChosen = std::as_const(LastChosen)](TSharedPtr<ITargetDeviceProxy> A, TSharedPtr<ITargetDeviceProxy> B)
+						   { return A->GetName() == LastChosen; });
 
 				// always use the first one, after sorting
 				FUIAction Action;
 				FText Tooltip;
-				GenerateDeviceProxyMenuParams(DeviceProxies[0], PlatformName, Action, Tooltip, ExternalOnClickDelegate);
 
-				if (DeviceProxies.Num() == 1)
+				GenerateDeviceProxyMenuParams(DeviceProxies[0], PlatformName, Action, Tooltip, ExternalOnClickDelegate);
+				
+ 				const ITargetPlatform* Platform = GetTargetPlatformManager()->FindTargetPlatform(PlatformName.ToString());
+				bool bGroupDevices = Platform->SupportsFeature(ETargetPlatformFeatures::ShowAsPlatformGroup);
+				if (DeviceProxies.Num() == 1 && !bGroupDevices)
 				{
 					DynamicSection.AddMenuEntry(
 						NAME_None,
@@ -1744,9 +1790,9 @@ void FTurnkeySupportModule::MakeQuickLaunchItems(class UToolMenu* Menu, FOnQuick
 				{
 					DynamicSection.AddSubMenu(
 						NAME_None,
-						MakeSdkStatusAttribute(PlatformName, DeviceProxies[0]),
+						MakeSdkStatusAttribute(PlatformName, bGroupDevices?nullptr:DeviceProxies[0]),
 						Tooltip,
-						FNewToolMenuDelegate::CreateLambda([TargetDeviceServicesModule, PlatformName, ExternalOnClickDelegate](UToolMenu* SubToolMenu)
+						FNewToolMenuDelegate::CreateLambda([TargetDeviceServicesModule, PlatformName, LastChosen, ExternalOnClickDelegate](UToolMenu* SubToolMenu)
 							{
 								FToolMenuSection& Section = SubToolMenu->AddSection(NAME_None);
 
@@ -1756,6 +1802,12 @@ void FTurnkeySupportModule::MakeQuickLaunchItems(class UToolMenu* Menu, FOnQuick
 								// for each one, put an entry (even the one that was in the outer menu, for less confusion)
 								for (const TSharedPtr<ITargetDeviceProxy>& Proxy : DeviceProxies)
 								{
+									// Skip over the top level menu item
+									if (LastChosen == Proxy->GetName())
+									{
+										continue;
+									}
+									
 									FUIAction SubAction;
 									FText SubTooltip;
 									GenerateDeviceProxyMenuParams(Proxy, PlatformName, SubAction, SubTooltip, ExternalOnClickDelegate);
@@ -1772,13 +1824,13 @@ void FTurnkeySupportModule::MakeQuickLaunchItems(class UToolMenu* Menu, FOnQuick
 						Action,
 						EUserInterfaceActionType::Check,
 						false,
-						MakePlatformSdkIconAttribute(PlatformName, nullptr),
+						MakePlatformSdkIconAttribute(PlatformName, bGroupDevices?nullptr:DeviceProxies[0]),
 						true
 						);
 				}
 
 				ITurnkeySupportModule& TurnkeySupport = ITurnkeySupportModule::Get();
-				// gath	er any unknown status devices to query at the end
+				// gather any unknown status devices to query at the end
 				for (const TSharedPtr<ITargetDeviceProxy>& Proxy : DeviceProxies)
 				{
 					FString DeviceId = Proxy->GetTargetDeviceId(NAME_None);
diff --git a/Engine/Source/Runtime/IOS/LaunchDaemonMessages/Classes/IOSMessageProtocol.h b/Engine/Source/Runtime/IOS/LaunchDaemonMessages/Classes/IOSMessageProtocol.h
--- a/Engine/Source/Runtime/IOS/LaunchDaemonMessages/Classes/IOSMessageProtocol.h
+++ b/Engine/Source/Runtime/IOS/LaunchDaemonMessages/Classes/IOSMessageProtocol.h
@@ -25,6 +25,9 @@ struct FIOSLaunchDaemonPong
 	UPROPERTY()
 	FString DeviceID;
 
+	UPROPERTY()
+	FString DeviceUDID;
+
 	UPROPERTY()
 	FString DeviceName;
 
@@ -34,6 +37,15 @@ struct FIOSLaunchDaemonPong
 	UPROPERTY()
 	FString DeviceType;
 
+	UPROPERTY()
+	FString DeviceOSVersion;
+
+	UPROPERTY()
+	FString DeviceModelId;
+
+	UPROPERTY()
+	FString DeviceConnectionType;
+
 	UPROPERTY()
 	bool bCanPowerOff;
 
@@ -53,11 +65,13 @@ struct FIOSLaunchDaemonPong
 		, bCanReboot(false)
 	{}
 
-	FIOSLaunchDaemonPong(FString InDeviceID, FString InDeviceName, FString InDeviceStatus, FString InDeviceType, bool bInCanPowerOff, bool bInCanPowerOn, bool bInCanReboot, bool bIsAuthorized)
+	FIOSLaunchDaemonPong(FString InDeviceID, FString InDeviceUDID, FString InDeviceName, FString InDeviceStatus, FString InDeviceType, FString InDeviceConnectionType, bool bInCanPowerOff, bool bInCanPowerOn, bool bInCanReboot, bool bIsAuthorized)
 		: DeviceID(InDeviceID)
+		, DeviceUDID(InDeviceUDID)
 		, DeviceName(InDeviceName)
 		, DeviceStatus(InDeviceStatus)
 		, DeviceType(InDeviceType)
+		, DeviceConnectionType(InDeviceConnectionType)
 		, bCanPowerOff(bInCanPowerOff)
 		, bCanPowerOn(bInCanPowerOn)
 		, bIsAuthorized(bIsAuthorized)
diff --git a/Engine/Source/Runtime/Launch/Private/IOS/GameLaunchDaemonMessageHandler.cpp b/Engine/Source/Runtime/Launch/Private/IOS/GameLaunchDaemonMessageHandler.cpp
--- a/Engine/Source/Runtime/Launch/Private/IOS/GameLaunchDaemonMessageHandler.cpp
+++ b/Engine/Source/Runtime/Launch/Private/IOS/GameLaunchDaemonMessageHandler.cpp
@@ -43,9 +43,11 @@ void FGameLaunchDaemonMessageHandler::HandlePingMessage(const FIOSLaunchDaemonPi
 		FMessageAddress MessageSender = Context->GetSender();
 
 		MessageEndpoint->Send(FMessageEndpoint::MakeMessage<FIOSLaunchDaemonPong>(FString(FPlatformProperties::PlatformName()) + (TARGET_IPHONE_SIMULATOR ? FString(TEXT("Simulator:")) : FString(TEXT("@"))) + FString(FPlatformProcess::ComputerName())
+			, FString(TEXT("udid"))
 			, FPlatformProcess::ComputerName()
 			, "Game_Running"
 			, [[UIDevice currentDevice] userInterfaceIdiom] == UIUserInterfaceIdiomPhone? "Phone" : "Tablet"
+			, FString(TEXT("local"))
 			, false
 			, false
 			, false
