Fix iOS remote build project settings category

diff --git a/Engine/Source/Developer/IOS/IOSPlatformEditor/Private/IOSTargetSettingsCustomization.cpp b/Engine/Source/Developer/IOS/IOSPlatformEditor/Private/IOSTargetSettingsCustomization.cpp
--- a/Engine/Source/Developer/IOS/IOSPlatformEditor/Private/IOSTargetSettingsCustomization.cpp
+++ b/Engine/Source/Developer/IOS/IOSPlatformEditor/Private/IOSTargetSettingsCustomization.cpp
@@ -893,7 +893,7 @@ void FIOSTargetSettingsCustomization::BuildPListSection(IDetailLayoutBuilder& De
 void FIOSTargetSettingsCustomization::BuildRemoteBuildingSection(IDetailLayoutBuilder& DetailLayout)
 {
 #if PLATFORM_WINDOWS
-	IDetailCategoryBuilder& BuildCategory = DetailLayout.EditCategory(TEXT("Build"));
+	IDetailCategoryBuilder& BuildCategory = DetailLayout.EditCategory(TEXT("Remote Build"));
 
 	// Sub group we wish to add remote building options to.
 	FText RemoteBuildingGroupName = LOCTEXT("RemoteBuildingGroupName", "Remote Build Options");
