Gray out some of iOS settings when using modernized Xcode

diff --git a/Engine/Source/Developer/IOS/IOSPlatformEditor/IOSPlatformEditor.Build.cs b/Engine/Source/Developer/IOS/IOSPlatformEditor/IOSPlatformEditor.Build.cs
--- a/Engine/Source/Developer/IOS/IOSPlatformEditor/IOSPlatformEditor.Build.cs
+++ b/Engine/Source/Developer/IOS/IOSPlatformEditor/IOSPlatformEditor.Build.cs
@@ -29,6 +29,7 @@ public class IOSPlatformEditor : ModuleRules
                 "AudioSettingsEditor",
                 "GameProjectGeneration",
 				"FreeImage",
+                "MacTargetPlatform",
 			}
 		);
 
diff --git a/Engine/Source/Developer/IOS/IOSPlatformEditor/Private/IOSTargetSettingsCustomization.cpp b/Engine/Source/Developer/IOS/IOSPlatformEditor/Private/IOSTargetSettingsCustomization.cpp
--- a/Engine/Source/Developer/IOS/IOSPlatformEditor/Private/IOSTargetSettingsCustomization.cpp
+++ b/Engine/Source/Developer/IOS/IOSPlatformEditor/Private/IOSTargetSettingsCustomization.cpp
@@ -38,6 +38,7 @@
 #include "IOSPlatformEditorModule.h"
 #include "HAL/PlatformFileManager.h"
 #include "UEFreeImage.h"
+#include "XcodeProjectSettings.h"
 
 #define LOCTEXT_NAMESPACE "IOSTargetSettings"
 DEFINE_LOG_CATEGORY_STATIC(LogIOSTargetSettings, Log, All);
@@ -337,6 +338,7 @@ void FIOSTargetSettingsCustomization::BuildPListSection(IDetailLayoutBuilder& De
 	BuildCategory.AddProperty(IOSTeamIDHandle)
 		.Visibility(EVisibility::Hidden);
 		BuildCategory.AddCustomRow(LOCTEXT("IOSTeamID", "IOSTeamID"), false)
+            .EditCondition(!UXcodeProjectSettings::IsUsingModernXcode(), nullptr)
 #if !PLATFORM_MAC
 			.Visibility(EVisibility::Hidden)
 #endif // PLATFORM_MAC
@@ -400,7 +402,20 @@ void FIOSTargetSettingsCustomization::BuildPListSection(IDetailLayoutBuilder& De
 			]
 		];*/
 
+    ProvisionCategory.AddCustomRow(LOCTEXT("ModernXcodeSigningLabel", "ModernXcodeSigning"), false)
+        .Visibility(UXcodeProjectSettings::IsUsingModernXcode() ? EVisibility::Visible : EVisibility::Hidden)
+        .WholeRowWidget
+        .MinDesiredWidth(0.f)
+        .MaxDesiredWidth(0.f)
+        .HAlign(HAlign_Fill)
+        [
+            SNew( STextBlock )
+            .Text( LOCTEXT( "ModernXcodeSigningText", "While Modern Xcode is enabled, code signing is handled by \"Xcode Projects\" section") )
+            .AutoWrapText( true )
+        ];
+    
 	ProvisionCategory.AddCustomRow(LOCTEXT("ProvisionLabel", "Provision"), false)
+        .EditCondition(!UXcodeProjectSettings::IsUsingModernXcode(), nullptr)
 		.WholeRowWidget
 		.MinDesiredWidth(0.f)
 		.MaxDesiredWidth(0.f)
@@ -558,6 +573,7 @@ void FIOSTargetSettingsCustomization::BuildPListSection(IDetailLayoutBuilder& De
 		];
 
 	ProvisionCategory.AddCustomRow(LOCTEXT("CertificateLabel", "Certificate"), false)
+        .EditCondition(!UXcodeProjectSettings::IsUsingModernXcode(), nullptr)
 		.WholeRowWidget
 		.MinDesiredWidth(0.f)
 		.MaxDesiredWidth(0.f)
@@ -738,6 +754,7 @@ void FIOSTargetSettingsCustomization::BuildPListSection(IDetailLayoutBuilder& De
 		Category.AddProperty(PropertyHandle) \
 		.Visibility(EVisibility::Hidden); \
 		Category.AddCustomRow(LOCTEXT("BundleIdentifier", "BundleIdentifier"), false) \
+        .EditCondition(!UXcodeProjectSettings::IsUsingModernXcode(), nullptr) \
 		.NameContent() \
 		[ \
 			SNew(SHorizontalBox) \
@@ -1276,6 +1293,7 @@ void FIOSTargetSettingsCustomization::BuildImageRow(IDetailLayoutBuilder& Detail
     if (Info.RequiredState == FPlatformIconInfo::Required)
     {
         Category.AddCustomRow(Info.IconName)
+            .EditCondition(!UXcodeProjectSettings::IsUsingModernXcode(), nullptr)
             .NameContent()
             [
                 SNew(SHorizontalBox)
@@ -1313,6 +1331,7 @@ void FIOSTargetSettingsCustomization::BuildImageRow(IDetailLayoutBuilder& Detail
     else
     {
         Category.AddCustomRow(Info.IconName)
+            .EditCondition(!UXcodeProjectSettings::IsUsingModernXcode(), nullptr)
             .NameContent()
             [
                 SNew(SHorizontalBox)
diff --git a/Engine/Source/Developer/Mac/MacTargetPlatform/Classes/XcodeProjectSettings.h b/Engine/Source/Developer/Mac/MacTargetPlatform/Classes/XcodeProjectSettings.h
--- a/Engine/Source/Developer/Mac/MacTargetPlatform/Classes/XcodeProjectSettings.h
+++ b/Engine/Source/Developer/Mac/MacTargetPlatform/Classes/XcodeProjectSettings.h
@@ -23,9 +23,24 @@ public:
 	
 	/**
 	 * Enable modernized Xcode, when building from Xcode, use native Xcode for bundle generation and archiving instead of UBT
+     * Restart required to apply this setting
 	 */
-	UPROPERTY(EditAnywhere, config, Category=Xcode, meta = (DisplayName = "Modernized Xcode"))
+	UPROPERTY(EditAnywhere, config, Category=Xcode, meta = (DisplayName = "Modernized Xcode", ConfigRestartRequired = true))
 	bool bUseModernXcode;
+    
+    UFUNCTION()
+    static bool IsUsingModernXcode()
+    {
+        auto DefaultObject = Cast<UXcodeProjectSettings>(UXcodeProjectSettings::StaticClass()->GetDefaultObject());
+        if (DefaultObject)
+        {
+            return DefaultObject->bUseModernXcode;
+        }
+        else
+        {
+            return false;
+        }
+    }
 	
 	/**
 	 * Team ID used for native Xcode code signing. This must be the 10 letters/numbers ID found in Membership Details tab found in https://developer.apple.com/account
diff --git a/Engine/Source/Runtime/IOS/IOSRuntimeSettings/Classes/IOSRuntimeSettings.h b/Engine/Source/Runtime/IOS/IOSRuntimeSettings/Classes/IOSRuntimeSettings.h
--- a/Engine/Source/Runtime/IOS/IOSRuntimeSettings/Classes/IOSRuntimeSettings.h
+++ b/Engine/Source/Runtime/IOS/IOSRuntimeSettings/Classes/IOSRuntimeSettings.h
@@ -412,7 +412,7 @@ public:
 	 * copied to the bundle app and the Launch screen image above will not be included in the app.
 	 * When using assets in your custom LaunchScreen.storyboard, add them in Build/IOS/Resources/Interface/Assets and they will be included.
 	 */
-	UPROPERTY(GlobalConfig, EditAnywhere, Category = LaunchScreen, meta = (DisplayName = "Custom Launchscreen Storyboard (experimental)"))
+	UPROPERTY(GlobalConfig, EditAnywhere, Category = LaunchScreen, meta = (DisplayName = "Custom Launchscreen Storyboard (experimental)", EditCondition = "!MacTargetPlatform.XcodeProjectSettings.IsUsingModernXcode"))
 	bool bCustomLaunchscreenStoryboard;
 
 	// Whether the app supports Facebook
@@ -434,7 +434,7 @@ public:
 	FString SigningCertificate;
 	
 	// Whether to use automatic signing through Xcode
-	UPROPERTY(GlobalConfig, EditAnywhere, Category = Build)
+	UPROPERTY(GlobalConfig, EditAnywhere, Category = Build, meta = (EditCondition = "!MacTargetPlatform.XcodeProjectSettings.IsUsingModernXcode"))
 	bool bAutomaticSigning;
 
 	// The team ID of the apple developer account to be used to autmatically sign IOS builds.
