- Fixed IPhones to not show up as VisionOS devices in the Platforms menu in the editor

diff --git a/Engine/Source/Developer/IOS/IOSTargetPlatform/Private/IOSTargetPlatform.cpp b/Engine/Source/Developer/IOS/IOSTargetPlatform/Private/IOSTargetPlatform.cpp
--- a/Engine/Source/Developer/IOS/IOSTargetPlatform/Private/IOSTargetPlatform.cpp
+++ b/Engine/Source/Developer/IOS/IOSTargetPlatform/Private/IOSTargetPlatform.cpp
@@ -32,7 +32,7 @@ FIOSTargetPlatform::FIOSTargetPlatform(bool bInIsTVOS, bool bInIsVisionOS, bool
 	// override the ini name up in the base classes, which will go into the FTargetPlatformInfo
 	: TNonDesktopTargetPlatformBase(bIsClientOnly, nullptr, bInIsTVOS ? TEXT("TVOS") : bInIsVisionOS ? TEXT("VisionOS") : nullptr)
 	, bIsTVOS(bInIsTVOS)
-	, bIsVisionOS(bInIsTVOS)
+	, bIsVisionOS(bInIsVisionOS)
 	, MobileShadingPath(0)
 	, bDistanceField(false)
 	, bMobileForwardEnableClusteredReflections(false)
@@ -48,7 +48,7 @@ FIOSTargetPlatform::FIOSTargetPlatform(bool bInIsTVOS, bool bInIsVisionOS, bool
 	// initialize the connected device detector
 	DeviceHelper.OnDeviceConnected().AddRaw(this, &FIOSTargetPlatform::HandleDeviceConnected);
 	DeviceHelper.OnDeviceDisconnected().AddRaw(this, &FIOSTargetPlatform::HandleDeviceDisconnected);
-	DeviceHelper.Initialize(bIsTVOS);
+	DeviceHelper.Initialize(bIsTVOS || bIsVisionOS);
 }
 
 
@@ -388,7 +388,13 @@ void FIOSTargetPlatform::HandleDeviceConnected(const FIOSLaunchDaemonPong& Messa
 	
 	if (!Device.IsValid())
 	{
-	if ((Message.DeviceType.Contains(TEXT("AppleTV")) && bIsTVOS) || (!Message.DeviceType.Contains(TEXT("AppleTV")) && !bIsTVOS))
+		bool bIsTVOSDevice = Message.DeviceType.Contains(TEXT("AppleTV"));
+		bool bIsVisionOSDevice = Message.DeviceType.Contains(TEXT("RealityDevice"));
+		bool bIsIOSDevice = !bIsTVOSDevice && !bIsVisionOSDevice;
+		
+		bool bIsIOS = !bIsTVOS && !bIsVisionOS;
+
+		if ((bIsIOS && bIsIOSDevice) || (bIsTVOS && bIsTVOSDevice) || (bIsVisionOS && bIsVisionOSDevice))
 		{
 			Device = MakeShareable(new FIOSTargetDevice(*this));
 
