Detect and store the available open iOS Simulators on mac
Part 1 of 3

diff --git a/Engine/Source/Developer/IOS/IOSTargetPlatform/IOSTargetPlatform.Build.cs b/Engine/Source/Developer/IOS/IOSTargetPlatform/IOSTargetPlatform.Build.cs
--- a/Engine/Source/Developer/IOS/IOSTargetPlatform/IOSTargetPlatform.Build.cs
+++ b/Engine/Source/Developer/IOS/IOSTargetPlatform/IOSTargetPlatform.Build.cs
@@ -17,6 +17,7 @@ public class IOSTargetPlatform : ModuleRules
 				"DesktopPlatform",
 				"LaunchDaemonMessages",
 				"Projects",
+				"Json",
 				"AudioPlatformConfiguration",
 				"Sockets",
 				"Networking"
diff --git a/Engine/Source/Developer/IOS/IOSTargetPlatform/Private/IOSDeviceHelper.cpp b/Engine/Source/Developer/IOS/IOSTargetPlatform/Private/IOSDeviceHelper.cpp
--- a/Engine/Source/Developer/IOS/IOSTargetPlatform/Private/IOSDeviceHelper.cpp
+++ b/Engine/Source/Developer/IOS/IOSTargetPlatform/Private/IOSDeviceHelper.cpp
@@ -10,165 +10,269 @@
 #include "Interfaces/ITargetPlatform.h"
 #include "Interfaces/IProjectManager.h"
 #include "Misc/MessageDialog.h"
+#include "Serialization/JsonReader.h"
+#include "Serialization/JsonSerializer.h"
 
 #define LOCTEXT_NAMESPACE "MiscIOSMessages"
 
 DEFINE_LOG_CATEGORY_STATIC(LogIOSDeviceHelper, Log, All);
 
-    enum DeviceConnectionInterface
-    {
-        NoValue = 0,
-        USB = 1,
-        Network = 2,
-        Max = 4
-    };
-
-    struct FDeviceNotificationCallbackInformation
-    {
-        FString DeviceID;
-        FString DeviceName;
-        FString DeviceUDID;
-        FString ProductType;
-        FString DeviceOSVersion;
-        DeviceConnectionInterface DeviceInterface;
-        uint32 msgType;
-		bool IsAuthorized;
-    };
-
-
-    struct LibIMobileDevice
-    {
-        FString DeviceID;
-        FString DeviceName;
-        FString DeviceUDID;
-        FString DeviceType;
-        FString DeviceOSVersion;
-        DeviceConnectionInterface DeviceInterface;
-		bool IsAuthorized;
-		bool IsDealtWith;
-    };
-
-    static TArray<LibIMobileDevice> GetLibIMobileDevices()
-    {
-        FString OutStdOut;
-        FString OutStdErr;
-        FString LibimobileDeviceId = GetLibImobileDeviceExe("idevice_id");
-        int ReturnCode;
-        // get the list of devices UDID
-        FPlatformProcess::ExecProcess(*LibimobileDeviceId, TEXT(""), &ReturnCode, &OutStdOut, &OutStdErr, NULL, true);
-        
-
-        TArray<LibIMobileDevice> ToReturn;
-        // separate out each line
-        TArray<FString> OngoingDeviceIds;
-        
-        OutStdOut.ParseIntoArray(OngoingDeviceIds, TEXT("\n"), true);
-        TArray<FString> DeviceStrings;
-        for (int32 StringIndex = 0; StringIndex < OngoingDeviceIds.Num(); ++StringIndex)
-        {
-            const FString& DeviceUDID = OngoingDeviceIds[StringIndex];
-            DeviceConnectionInterface OngoingDeviceInterface = DeviceConnectionInterface::NoValue;
-
-            FString OutStdOutInfo;
-            FString OutStdErrInfo;
-            FString LibimobileDeviceInfo = GetLibImobileDeviceExe("ideviceinfo");
-            int ReturnCodeInfo;
-            FString Arguments;
-
-            if (OngoingDeviceIds[StringIndex].Contains("USB"))
-            {
-                OngoingDeviceInterface = DeviceConnectionInterface::USB;
-            }
-            else if (OngoingDeviceIds[StringIndex].Contains("Network"))
-            {
-                OngoingDeviceInterface = DeviceConnectionInterface::Network;
-            }
-            OngoingDeviceIds[StringIndex].Split(TEXT(" "), &OngoingDeviceIds[StringIndex], nullptr, ESearchCase::CaseSensitive, ESearchDir::FromStart);
-
-            if (OngoingDeviceInterface == DeviceConnectionInterface::USB)
-            {
-                Arguments = "-u " + DeviceUDID;
-            }
-            else if (OngoingDeviceInterface == DeviceConnectionInterface::Network)
-            {
-                Arguments = "-n -u " + DeviceUDID;
-            }
-
-            FPlatformProcess::ExecProcess(*LibimobileDeviceInfo, *Arguments, &ReturnCodeInfo, &OutStdOutInfo, &OutStdErrInfo, NULL, true);
-            
-			LibIMobileDevice ToAdd;
-
-            // ideviceinfo can fail when the connected device is untrusted. It can be "Pairing dialog response pending (-19)", "Invalid HostID (-21)" or "User denied pairing (-18)"
-            // the only thing we can do is to make sure the Trust popup is correctly displayed.
-            if (OutStdErrInfo.Contains("ERROR: "))
-            {
-				if (OutStdErrInfo.Contains("Could not connect to lockdownd"))
-				{
-					// UE_LOG(LogIOSDeviceHelper, Warning, TEXT("Could not pair with connected iOS/tvOS device. Trust this computer by accepting the popup on device."));
-					FString LibimobileDevicePair = GetLibImobileDeviceExe("idevicepair");
-					FString PairArguments = "-u " + DeviceUDID + " pair";
-					FPlatformProcess::ExecProcess(*LibimobileDevicePair, *PairArguments, &ReturnCodeInfo, &OutStdOutInfo, &OutStdErrInfo, NULL, true);
-				}
-				else
-				{
-					UE_LOG(LogIOSDeviceHelper, Warning, TEXT("Libimobile call failed : %s"), *OutStdErrInfo);
-				}
-				OutStdOutInfo.Empty();
-				OutStdErrInfo.Empty();
-				ToAdd.IsAuthorized = false;
+enum DeviceConnectionInterface
+{
+	NoValue = 0,
+	USB = 1,
+	Network = 2,
+	Simulator = 3,
+	Max = 4
+};
+
+struct FDeviceNotificationCallbackInformation
+{
+	FString DeviceID;
+	FString DeviceName;
+	FString DeviceUDID;
+	FString ProductType;
+	FString DeviceOSVersion;
+	DeviceConnectionInterface DeviceInterface;
+	uint32 msgType;
+	bool IsAuthorized;
+};
+
+
+struct LibIMobileDevice
+{
+	FString DeviceID;
+	FString DeviceName;
+	FString DeviceUDID;
+	FString DeviceType;
+	FString DeviceOSVersion;
+	DeviceConnectionInterface DeviceInterface;
+	bool IsAuthorized;
+	bool IsDealtWith;
+};
+
+const FString StringifyDeviceConnection(DeviceConnectionInterface Interface)
+{
+	switch(Interface)
+	{
+		case DeviceConnectionInterface::NoValue:
+			return TEXT("NoValue");
+
+		case DeviceConnectionInterface::Network:
+			return TEXT("Network");
+
+		case DeviceConnectionInterface::USB:
+			return TEXT("USB");
+
+		case DeviceConnectionInterface::Simulator:
+			return TEXT("Simulator");
+
+		default:
+			UE_LOG(LogIOSDeviceHelper, Warning, TEXT("Unknown DeviceConnectionInterface type:%d"), Interface);
+			return TEXT("NoValue");
+	}	
+}
+
+static TArray<LibIMobileDevice> GetLibIMobileDevices()
+{
+	FString OutStdOut;
+	FString OutStdErr;
+	FString LibimobileDeviceId = GetLibImobileDeviceExe("idevice_id");
+	int ReturnCode;
+	// get the list of devices UDID
+	FPlatformProcess::ExecProcess(*LibimobileDeviceId, TEXT(""), &ReturnCode, &OutStdOut, &OutStdErr, NULL, true);
+	
+	
+	TArray<LibIMobileDevice> ToReturn;
+	// separate out each line
+	TArray<FString> OngoingDeviceIds;
+	
+	OutStdOut.ParseIntoArray(OngoingDeviceIds, TEXT("\n"), true);
+	TArray<FString> DeviceStrings;
+	for (int32 StringIndex = 0; StringIndex < OngoingDeviceIds.Num(); ++StringIndex)
+	{
+		const FString& DeviceUDID = OngoingDeviceIds[StringIndex];
+		DeviceConnectionInterface OngoingDeviceInterface = DeviceConnectionInterface::NoValue;
+		
+		FString OutStdOutInfo;
+		FString OutStdErrInfo;
+		FString LibimobileDeviceInfo = GetLibImobileDeviceExe("ideviceinfo");
+		int ReturnCodeInfo;
+		FString Arguments;
+		
+		if (OngoingDeviceIds[StringIndex].Contains("USB"))
+		{
+			OngoingDeviceInterface = DeviceConnectionInterface::USB;
+		}
+		else if (OngoingDeviceIds[StringIndex].Contains("Network"))
+		{
+			OngoingDeviceInterface = DeviceConnectionInterface::Network;
+		}
+		OngoingDeviceIds[StringIndex].Split(TEXT(" "), &OngoingDeviceIds[StringIndex], nullptr, ESearchCase::CaseSensitive, ESearchDir::FromStart);
+		
+		if (OngoingDeviceInterface == DeviceConnectionInterface::USB)
+		{
+			Arguments = "-u " + DeviceUDID;
+		}
+		else if (OngoingDeviceInterface == DeviceConnectionInterface::Network)
+		{
+			Arguments = "-n -u " + DeviceUDID;
+		}
+		
+		FPlatformProcess::ExecProcess(*LibimobileDeviceInfo, *Arguments, &ReturnCodeInfo, &OutStdOutInfo, &OutStdErrInfo, NULL, true);
+		
+		LibIMobileDevice ToAdd;
+		
+		// ideviceinfo can fail when the connected device is untrusted. It can be "Pairing dialog response pending (-19)", "Invalid HostID (-21)" or "User denied pairing (-18)"
+		// the only thing we can do is to make sure the Trust popup is correctly displayed.
+		if (OutStdErrInfo.Contains("ERROR: "))
+		{
+			if (OutStdErrInfo.Contains("Could not connect to lockdownd"))
+			{
+				// UE_LOG(LogIOSDeviceHelper, Warning, TEXT("Could not pair with connected iOS/tvOS device. Trust this computer by accepting the popup on device."));
+				FString LibimobileDevicePair = GetLibImobileDeviceExe("idevicepair");
+				FString PairArguments = "-u " + DeviceUDID + " pair";
+				FPlatformProcess::ExecProcess(*LibimobileDevicePair, *PairArguments, &ReturnCodeInfo, &OutStdOutInfo, &OutStdErrInfo, NULL, true);
 			}
 			else
 			{
-				ToAdd.IsAuthorized = true;
+				UE_LOG(LogIOSDeviceHelper, Warning, TEXT("Libimobile call failed : %s"), *OutStdErrInfo);
 			}
-
-			// parse product type and device name
-			FString DeviceName;
-			OutStdOutInfo.Split(TEXT("DeviceName: "), nullptr, &DeviceName, ESearchCase::CaseSensitive, ESearchDir::FromStart);
-			DeviceName.Split(LINE_TERMINATOR, &DeviceName, nullptr, ESearchCase::CaseSensitive, ESearchDir::FromStart);
-			if (!ToAdd.IsAuthorized)
+			OutStdOutInfo.Empty();
+			OutStdErrInfo.Empty();
+			ToAdd.IsAuthorized = false;
+		}
+		else
+		{
+			ToAdd.IsAuthorized = true;
+		}
+		
+		// parse product type and device name
+		FString DeviceName;
+		OutStdOutInfo.Split(TEXT("DeviceName: "), nullptr, &DeviceName, ESearchCase::CaseSensitive, ESearchDir::FromStart);
+		DeviceName.Split(LINE_TERMINATOR, &DeviceName, nullptr, ESearchCase::CaseSensitive, ESearchDir::FromStart);
+		if (!ToAdd.IsAuthorized)
+		{
+			DeviceName = LOCTEXT("IosTvosUnauthorizedDevice", "iOS / tvOS (Unauthorized)").ToString();
+		}
+		else
+		{
+			if (OngoingDeviceInterface == DeviceConnectionInterface::Network)
 			{
-				DeviceName = LOCTEXT("IosTvosUnauthorizedDevice", "iOS / tvOS (Unauthorized)").ToString();
+				DeviceName += " [Wifi]";
 			}
-			else
+		}
+		
+		FString ProductType;
+		OutStdOutInfo.Split(TEXT("ProductType: "), nullptr, &ProductType, ESearchCase::CaseSensitive, ESearchDir::FromStart);
+		ProductType.Split(LINE_TERMINATOR, &ProductType, nullptr, ESearchCase::CaseSensitive, ESearchDir::FromStart);
+		
+		FString OSVersion; // iOS/iPad OS Version
+		OutStdOutInfo.Split(TEXT("ProductVersion: "), nullptr, &OSVersion, ESearchCase::CaseSensitive, ESearchDir::FromStart);
+		OSVersion.Split(LINE_TERMINATOR, &OSVersion, nullptr, ESearchCase::CaseSensitive, ESearchDir::FromStart);
+		
+		FString DeviceID = FString::Printf(TEXT("%s@%s"),
+										   ProductType.Contains(TEXT("AppleTV")) ? TEXT("TVOS") : TEXT("IOS"),
+										   *DeviceUDID);
+		
+		ToAdd.DeviceID = DeviceID;
+		ToAdd.DeviceUDID = DeviceUDID;
+		ToAdd.DeviceName = DeviceName;
+		ToAdd.DeviceType = ProductType;
+		ToAdd.DeviceOSVersion = OSVersion;
+		ToAdd.DeviceInterface = OngoingDeviceInterface;
+		ToAdd.IsDealtWith = false;
+		ToReturn.Add(ToAdd);
+	}
+	return ToReturn;
+}
+
+static TArray<LibIMobileDevice> GetSimulatorDevices()
+{
+	FString OutStdOut;
+	FString OutStdErr;
+	TArray<LibIMobileDevice> ToReturn;
+
+	int ReturnCode;
+	FString simCommand = TEXT("/usr/bin/xcrun");
+	FString simParams = TEXT("simctl list -je devices available");
+	FPlatformProcess::ExecProcess(*simCommand, *simParams, &ReturnCode, &OutStdOut, &OutStdErr, NULL, true);
+
+	if (OutStdOut.Len() == 0)
+	{
+		UE_LOG(LogIOSDeviceHelper, Warning, TEXT("Unable to access simctl app.  Is Xcode installed?"));
+		return ToReturn;
+	}
+
+	// decode the json object
+	TSharedPtr<FJsonObject> JsonObject;
+	TSharedRef<TJsonReader<>> Reader = TJsonReaderFactory<>::Create(*OutStdOut);//ANSI_TO_TCHAR((const ANSICHAR *)Msg.GetData()));
+	if (!FJsonSerializer::Deserialize(Reader, JsonObject) || !JsonObject.IsValid())
+	{
+		UE_LOG(LogIOSDeviceHelper, Error, TEXT("Failed to parse message from SimCtl utility: %s"), *Reader->GetErrorMessage());
+		return ToReturn;
+	}
+
+	const TSharedPtr<FJsonObject>* SimDevices;
+	if (JsonObject->TryGetObjectField(TEXT("devices"), /*out*/ SimDevices))
+	{
+		// This should be a list of iOS/tvOS/xrOS versions, each with an array of simulators for that version
+		for (const auto& OSSimPair : (*SimDevices)->Values)
+		{
+			// We only track iOS, for now
+			// The Key is format is: com.apple.CoreSimulator.SimRuntime.iOS-17-2
+			int IOSIndex = OSSimPair.Key.Find(TEXT("iOS"));
+			if (IOSIndex > 0)
 			{
-				if (OngoingDeviceInterface == DeviceConnectionInterface::Network)
+				FString OSVer = OSSimPair.Key.RightChop(IOSIndex+4);
+				OSVer.ReplaceInline(TEXT("-"), TEXT("."));
+
+				LibIMobileDevice ToAdd;
+				ToAdd.IsAuthorized = true;
+				ToAdd.DeviceOSVersion = OSVer;
+
+				const TArray<TSharedPtr<FJsonValue> >* SimDataArray;
+				if (OSSimPair.Value->TryGetArray(SimDataArray))
 				{
-					DeviceName += " [Wifi]";
+					for (const auto& SimData : *SimDataArray)
+					{
+						const TSharedPtr<FJsonObject>* SimDataObj;
+						if (SimData->TryGetObject(SimDataObj))
+						{
+							// For now, we only care about booted devices
+							const FString DeviceState = (*SimDataObj)->GetStringField(TEXT("state"));
+							if (DeviceState == TEXT("Booted"))
+							{
+								const FString DeviceName = (*SimDataObj)->GetStringField(TEXT("name"));
+								ToAdd.DeviceUDID = (*SimDataObj)->GetStringField(TEXT("udid"));
+								ToAdd.DeviceType = (*SimDataObj)->GetStringField(TEXT("name"));
+								ToAdd.DeviceName = DeviceName;
+							
+								FString DeviceID = FString::Printf(TEXT("IOS@%s"), *ToAdd.DeviceUDID);
+								ToAdd.DeviceID = DeviceID;
+								ToAdd.DeviceInterface = DeviceConnectionInterface::Simulator;
+							
+								ToReturn.Add(ToAdd);
+							}
+						}
+					}
 				}
 			}
 			
-			FString ProductType;
-			OutStdOutInfo.Split(TEXT("ProductType: "), nullptr, &ProductType, ESearchCase::CaseSensitive, ESearchDir::FromStart);
-			ProductType.Split(LINE_TERMINATOR, &ProductType, nullptr, ESearchCase::CaseSensitive, ESearchDir::FromStart);
-			
-			FString OSVersion; // iOS/iPad OS Version
-			OutStdOutInfo.Split(TEXT("ProductVersion: "), nullptr, &OSVersion, ESearchCase::CaseSensitive, ESearchDir::FromStart);
-			OSVersion.Split(LINE_TERMINATOR, &OSVersion, nullptr, ESearchCase::CaseSensitive, ESearchDir::FromStart);
-			
-			FString DeviceID = FString::Printf(TEXT("%s@%s"),
-											   ProductType.Contains(TEXT("AppleTV")) ? TEXT("TVOS") : TEXT("IOS"),
-											   *DeviceUDID);
-			
-			ToAdd.DeviceID = DeviceID;
-			ToAdd.DeviceUDID = DeviceUDID;
-			ToAdd.DeviceName = DeviceName;
-			ToAdd.DeviceType = ProductType;
-			ToAdd.DeviceOSVersion = OSVersion;
-			ToAdd.DeviceInterface = OngoingDeviceInterface;
-			ToAdd.IsDealtWith = false;
-			ToReturn.Add(ToAdd);
 		}
-		return ToReturn;
+
 	}
+	
+	return ToReturn;
+}
 
 class FIOSDevice
 {
 public:
 	FIOSDevice(FString InID, FString InName, DeviceConnectionInterface InConnectionType)
-		: UDID(InID)
-		, Name(InName)
-		, ConnectionType(InConnectionType)
+	: UDID(InID)
+	, Name(InName)
+	, ConnectionType(InConnectionType)
 	{
 	}
 	
@@ -203,180 +307,210 @@ DECLARE_MULTICAST_DELEGATE_OneParam(FDeviceNotification, void*)
 #define        RECHECK_COUNTER_RESET            12
 
 class FDeviceQueryTask
-    : public FRunnable
+: public FRunnable
 {
 public:
-    FDeviceQueryTask()
-        : Stopping(false)
-        , bCheckDevices(true)
-        , NeedSDKCheck(true)
-        , RetryQuery(5)
-    {}
-
-    virtual bool Init() override
-    {
-        return true;
-    }
-
-    virtual uint32 Run() override
-    {
-        while (!Stopping)
-        {
-            if (IsEngineExitRequested())
-            {
-                break;
-            }
-            if (GetTargetPlatformManager())
-            {
-                FString OutTutorialPath;
-                const ITargetPlatform* Platform = GetTargetPlatformManager()->FindTargetPlatform(TEXT("IOS"));
-                if (Platform)
-                {
-                    if (Platform->IsSdkInstalled(false, OutTutorialPath))
-                    {
-                        break;
-                    }
-                }
-                Enable(false);
-                return 0;
-            }
-            else
-            {
-                FPlatformProcess::Sleep(1.0f);
-            }
-        }
-        int RecheckCounter = RECHECK_COUNTER_RESET;
-        while (!Stopping)
-        {
-            if (IsEngineExitRequested())
-            {
-                break;
-            }
-            if (bCheckDevices)
-            {
+	FDeviceQueryTask()
+	: Stopping(false)
+	, bCheckDevices(true)
+	, NeedSDKCheck(true)
+	, RetryQuery(5)
+	{}
+	
+	virtual bool Init() override
+	{
+		return true;
+	}
+	
+	virtual uint32 Run() override
+	{
+		while (!Stopping)
+		{
+			if (IsEngineExitRequested())
+			{
+				break;
+			}
+			if (GetTargetPlatformManager())
+			{
+				FString OutTutorialPath;
+				const ITargetPlatform* Platform = GetTargetPlatformManager()->FindTargetPlatform(TEXT("IOS"));
+				if (Platform)
+				{
+					if (Platform->IsSdkInstalled(false, OutTutorialPath))
+					{
+						break;
+					}
+				}
+				Enable(false);
+				return 0;
+			}
+			else
+			{
+				FPlatformProcess::Sleep(1.0f);
+			}
+		}
+		int RecheckCounter = RECHECK_COUNTER_RESET;
+		while (!Stopping)
+		{
+			if (IsEngineExitRequested())
+			{
+				break;
+			}
+			if (bCheckDevices)
+			{
 #if WITH_EDITOR
-                if (!IsRunningCommandlet())
-                {
-                    //if (NeedSDKCheck)
-                    //{
-                    //    NeedSDKCheck = false;
-                    //    FProjectStatus ProjectStatus;
-                    //    if (!IProjectManager::Get().QueryStatusForCurrentProject(ProjectStatus) || (!ProjectStatus.IsTargetPlatformSupported(TEXT("IOS")) && !ProjectStatus.IsTargetPlatformSupported(TEXT("TVOS"))))
-                    //    {
-                    //        Enable(false);
-                    //    }
-                    //}
-                    //else
-                    {
-                        // BHP - Turning off device check to prevent it from interfering with packaging
-                        QueryDevices();
-                    }
-                }
+				if (!IsRunningCommandlet())
+				{
+					//if (NeedSDKCheck)
+					//{
+					//    NeedSDKCheck = false;
+					//    FProjectStatus ProjectStatus;
+					//    if (!IProjectManager::Get().QueryStatusForCurrentProject(ProjectStatus) || (!ProjectStatus.IsTargetPlatformSupported(TEXT("IOS")) && !ProjectStatus.IsTargetPlatformSupported(TEXT("TVOS"))))
+					//    {
+					//        Enable(false);
+					//    }
+					//}
+					//else
+					{
+						// BHP - Turning off device check to prevent it from interfering with packaging
+						QueryDevices();
+					}
+				}
 #else
-				 QueryDevices();
+				QueryDevices();
 #endif
-            }
-            RecheckCounter--;
-            if (RecheckCounter < 0)
-            {
-                RecheckCounter = RECHECK_COUNTER_RESET;
-                bCheckDevices = true;
-                NeedSDKCheck = true;
-            }
-
-            FPlatformProcess::Sleep(5.0f);
-        }
-
-        return 0;
-    }
-
-    virtual void Stop() override
-    {
-        Stopping = true;
-    }
-
-    virtual void Exit() override
-    {}
-
-    FDeviceNotification& OnDeviceNotification()
-    {
-        return DeviceNotification;
-    }
-
-    void Enable(bool bInCheckDevices)
-    {
-        bCheckDevices = bInCheckDevices;
-    }
-
+			}
+			RecheckCounter--;
+			if (RecheckCounter < 0)
+			{
+				RecheckCounter = RECHECK_COUNTER_RESET;
+				bCheckDevices = true;
+				NeedSDKCheck = true;
+			}
+			
+			FPlatformProcess::Sleep(5.0f);
+		}
+		
+		return 0;
+	}
+	
+	virtual void Stop() override
+	{
+		Stopping = true;
+	}
+	
+	virtual void Exit() override
+	{}
+	
+	FDeviceNotification& OnDeviceNotification()
+	{
+		return DeviceNotification;
+	}
+	
+	void Enable(bool bInCheckDevices)
+	{
+		bCheckDevices = bInCheckDevices;
+	}
+	
 private:
-    
-    void NotifyDeviceChange(LibIMobileDevice& Device, bool bAdd)
-    {
-        FDeviceNotificationCallbackInformation CallbackInfo;
-
-        if (bAdd)
-        {
-            CallbackInfo.DeviceID = Device.DeviceID;
-            CallbackInfo.DeviceName = Device.DeviceName;
-            CallbackInfo.DeviceUDID = Device.DeviceUDID;
-            CallbackInfo.DeviceInterface = Device.DeviceInterface;
-            CallbackInfo.ProductType = Device.DeviceType;
-            CallbackInfo.DeviceOSVersion = Device.DeviceOSVersion;
-            CallbackInfo.msgType = 1;
+	
+	void NotifyDeviceChange(LibIMobileDevice& Device, bool bAdd)
+	{
+		FDeviceNotificationCallbackInformation CallbackInfo;
+		
+		if (bAdd)
+		{
+			CallbackInfo.DeviceID = Device.DeviceID;
+			CallbackInfo.DeviceName = Device.DeviceName;
+			CallbackInfo.DeviceUDID = Device.DeviceUDID;
+			CallbackInfo.DeviceInterface = Device.DeviceInterface;
+			CallbackInfo.ProductType = Device.DeviceType;
+			CallbackInfo.DeviceOSVersion = Device.DeviceOSVersion;
+			CallbackInfo.msgType = 1;
 			CallbackInfo.IsAuthorized = Device.IsAuthorized;
-        }
-        else
-        {
-            CallbackInfo.DeviceID = Device.DeviceID;
+		}
+		else
+		{
+			CallbackInfo.DeviceID = Device.DeviceID;
 			CallbackInfo.DeviceUDID = Device.DeviceUDID;
 			CallbackInfo.DeviceInterface = Device.DeviceInterface;
-            CallbackInfo.msgType = 2;
-            DeviceNotification.Broadcast(&CallbackInfo);
-
-        }
-        DeviceNotification.Broadcast(&CallbackInfo);
-    }
-    
-    void QueryDevices()
-    {
-        FString OutStdOut;
-        FString OutStdErr;
+			CallbackInfo.msgType = 2;
+			DeviceNotification.Broadcast(&CallbackInfo);
+			
+		}
+		DeviceNotification.Broadcast(&CallbackInfo);
+	}
+	
+	void QueryDevices()
+	{
+		TArray<LibIMobileDevice> SimulatorDevices;
+		bool HasSimDevices = false;
+		bool HasDevices = true;
+
+		bool bEnableSimulatorSupport = false;
+		GConfig->GetBool(TEXT("/Script/IOSRuntimeSettings.IOSRuntimeSettings"), TEXT("bEnableSimulatorSupport"), bEnableSimulatorSupport, GEngineIni);
+		if (bEnableSimulatorSupport)
+		{
+			SimulatorDevices = GetSimulatorDevices();
+			HasSimDevices = SimulatorDevices.Num() > 0;
+		}
+
 		FString LibimobileDeviceId = GetLibImobileDeviceExe("idevice_id");
-		int ReturnCode;
 		if (LibimobileDeviceId.Len() == 0)
 		{
 			UE_LOG(LogIOSDeviceHelper, Log, TEXT("idevice_id (iOS device detection) executable missing. Turning off iOS/tvOS device detection."));
+			HasDevices = false;
 			Enable(false);
-			return;
+			if (!HasSimDevices)
+			{
+				return;
+			}
 		}
 
-        // get the list of devices UDID
-        FPlatformProcess::ExecProcess(*LibimobileDeviceId, TEXT(""), &ReturnCode, &OutStdOut, &OutStdErr, NULL, true);
-        if (OutStdOut.Len() == 0)
-        {
-            RetryQuery--;
-            if (RetryQuery < 0)
-            {
-                UE_LOG(LogIOSDeviceHelper, Verbose, TEXT("IOS device listing is disabled for 1 minute (too many failed attempts)!"));
-                Enable(false);
-            }
-			for (LibIMobileDevice device : CachedDevices)
+		// get the list of devices UDID
+		if (HasDevices)
+		{
+			FString OutStdOut;
+			FString OutStdErr;
+			int ReturnCode;
+			FPlatformProcess::ExecProcess(*LibimobileDeviceId, TEXT(""), &ReturnCode, &OutStdOut, &OutStdErr, NULL, true);
+			if (OutStdOut.Len() == 0)
 			{
-				NotifyDeviceChange(device, false);
+				RetryQuery--;
+				if (RetryQuery < 0)
+				{
+					UE_LOG(LogIOSDeviceHelper, Verbose, TEXT("IOS device listing is disabled for 1 minute (too many failed attempts)!"));
+					Enable(false);
+				}
+				for (LibIMobileDevice device : CachedDevices)
+				{
+					NotifyDeviceChange(device, false);
+				}
+				CachedDevices.Empty();
+				HasDevices = false;
+				if (!HasSimDevices)
+				{
+					return;
+				}
 			}
-			CachedDevices.Empty();
-            return;
-        }
-        RetryQuery = 5;
-    
-        TArray<LibIMobileDevice> ParsedDevices = GetLibIMobileDevices();
-
-        for (int32 Index = 0; Index < ParsedDevices.Num(); ++Index)
-        {
-			LibIMobileDevice* Found = CachedDevices.FindByPredicate(
-				[&](LibIMobileDevice Element) {
+			RetryQuery = 5;
+		}
+
+		TArray<LibIMobileDevice> ParsedDevices;
+		if (HasDevices)
+		{
+			ParsedDevices = GetLibIMobileDevices();
+		}
+
+		ParsedDevices.Append(SimulatorDevices);
+		
+		for (int32 Index = 0; Index < ParsedDevices.Num(); ++Index)
+		{
+			LibIMobileDevice *Found = CachedDevices.FindByPredicate(
+				[&](LibIMobileDevice Element)
+				{
 					return (Element.DeviceUDID == ParsedDevices[Index].DeviceUDID &&
-						Element.DeviceInterface == ParsedDevices[Index].DeviceInterface);
+							Element.DeviceInterface == ParsedDevices[Index].DeviceInterface);
 				});
 			if (Found != nullptr)
 			{
@@ -391,27 +525,27 @@ private:
 			{
 				NotifyDeviceChange(ParsedDevices[Index], true);
 			}
-	
-        }
-
-        for (int32 Index = 0; Index < CachedDevices.Num(); ++Index)
-        {
-            if (!CachedDevices[Index].IsDealtWith)
-            {
-                NotifyDeviceChange(CachedDevices[Index], false);
-            }
-        }
-
-        CachedDevices.Empty();
+			
+		}
+		
+		for (int32 Index = 0; Index < CachedDevices.Num(); ++Index)
+		{
+			if (!CachedDevices[Index].IsDealtWith)
+			{
+				NotifyDeviceChange(CachedDevices[Index], false);
+			}
+		}
+		
+		CachedDevices.Empty();
 		CachedDevices = ParsedDevices;
-    }
-
-    bool Stopping;
-    bool bCheckDevices;
-    bool NeedSDKCheck;
-    int RetryQuery;
-    TArray<LibIMobileDevice> CachedDevices;
-    FDeviceNotification DeviceNotification;
+	}
+	
+	bool Stopping;
+	bool bCheckDevices;
+	bool NeedSDKCheck;
+	int RetryQuery;
+	TArray<LibIMobileDevice> CachedDevices;
+	FDeviceNotification DeviceNotification;
 };
 
 /* FIOSDeviceHelper structors
@@ -424,86 +558,86 @@ static FTickerDelegate TickDelegate;
 
 bool FIOSDeviceHelper::MessageTickDelegate(float DeltaTime)
 {
-    QUICK_SCOPE_CYCLE_COUNTER(STAT_FIOSDeviceHelper_MessageTickDelegate);
-
-    for (int Index = 0; Index < NotificationMessages.Num(); ++Index)
-    {
-        FDeviceNotificationCallbackInformation cbi = NotificationMessages[Index];
-        FIOSDeviceHelper::DeviceCallback(&cbi);
-    }
-    NotificationMessages.Empty();
-
-    return true;
+	QUICK_SCOPE_CYCLE_COUNTER(STAT_FIOSDeviceHelper_MessageTickDelegate);
+	
+	for (int Index = 0; Index < NotificationMessages.Num(); ++Index)
+	{
+		FDeviceNotificationCallbackInformation cbi = NotificationMessages[Index];
+		FIOSDeviceHelper::DeviceCallback(&cbi);
+	}
+	NotificationMessages.Empty();
+	
+	return true;
 }
 
 void FIOSDeviceHelper::Initialize(bool bIsTVOS)
 {
-    if(!bIsTVOS)
-    {
-        // add the message pump
-        TickDelegate = FTickerDelegate::CreateStatic(MessageTickDelegate);
-        FTSTicker::GetCoreTicker().AddTicker(TickDelegate, 5.0f);
-
-        // kick off a thread to query for connected devices
-        QueryTask = new FDeviceQueryTask();
-        QueryTask->OnDeviceNotification().AddStatic(FIOSDeviceHelper::DeviceCallback);
-
-        static int32 QueryTaskCount = 1;
-        if (QueryTaskCount == 1)
-        {
-            // create the socket subsystem (loadmodule in game thread)
-            ISocketSubsystem* SSS = ISocketSubsystem::Get();
-            QueryThread = FRunnableThread::Create(QueryTask, *FString::Printf(TEXT("FIOSDeviceHelper.QueryTask_%d"), QueryTaskCount++), 128 * 1024, TPri_Normal);
-        }
-    }
+	if(!bIsTVOS)
+	{
+		// add the message pump
+		TickDelegate = FTickerDelegate::CreateStatic(MessageTickDelegate);
+		FTSTicker::GetCoreTicker().AddTicker(TickDelegate, 5.0f);
+		
+		// kick off a thread to query for connected devices
+		QueryTask = new FDeviceQueryTask();
+		QueryTask->OnDeviceNotification().AddStatic(FIOSDeviceHelper::DeviceCallback);
+		
+		static int32 QueryTaskCount = 1;
+		if (QueryTaskCount == 1)
+		{
+			// create the socket subsystem (loadmodule in game thread)
+			ISocketSubsystem* SSS = ISocketSubsystem::Get();
+			QueryThread = FRunnableThread::Create(QueryTask, *FString::Printf(TEXT("FIOSDeviceHelper.QueryTask_%d"), QueryTaskCount++), 128 * 1024, TPri_Normal);
+		}
+	}
 }
 
 void FIOSDeviceHelper::DeviceCallback(void* CallbackInfo)
 {
-    struct FDeviceNotificationCallbackInformation* cbi = (FDeviceNotificationCallbackInformation*)CallbackInfo;
-
-    if (!IsInGameThread())
-    {
-        NotificationMessages.Add(*cbi);
-    }
-    else
-    {
-        switch(cbi->msgType)
-        {
-        case 1:
-            FIOSDeviceHelper::DoDeviceConnect(CallbackInfo);
-            break;
-
-        case 2:
-            FIOSDeviceHelper::DoDeviceDisconnect(CallbackInfo);
-            break;
-        }
-    }
+	struct FDeviceNotificationCallbackInformation* cbi = (FDeviceNotificationCallbackInformation*)CallbackInfo;
+	
+	if (!IsInGameThread())
+	{
+		NotificationMessages.Add(*cbi);
+	}
+	else
+	{
+		switch(cbi->msgType)
+		{
+			case 1:
+				FIOSDeviceHelper::DoDeviceConnect(CallbackInfo);
+				break;
+				
+			case 2:
+				FIOSDeviceHelper::DoDeviceDisconnect(CallbackInfo);
+				break;
+		}
+	}
 }
 
 void FIOSDeviceHelper::DoDeviceConnect(void* CallbackInfo)
 {
-    // connect to the device
-    struct FDeviceNotificationCallbackInformation* cbi = (FDeviceNotificationCallbackInformation*)CallbackInfo;
-    FIOSDevice* Device = new FIOSDevice(cbi->DeviceUDID, cbi->DeviceName, cbi->DeviceInterface);
-
-    // fire the event
-    FIOSLaunchDaemonPong Event;
+	// connect to the device
+	struct FDeviceNotificationCallbackInformation* cbi = (FDeviceNotificationCallbackInformation*)CallbackInfo;
+	FIOSDevice* Device = new FIOSDevice(cbi->DeviceUDID, cbi->DeviceName, cbi->DeviceInterface);
+	
+	// fire the event
+	FIOSLaunchDaemonPong Event;
 	Event.DeviceID = cbi->DeviceID;
-    Event.DeviceUDID = cbi->DeviceUDID;
-    Event.DeviceName = cbi->DeviceName;
-    Event.DeviceType = cbi->ProductType;
-    Event.DeviceOSVersion = cbi->DeviceOSVersion;
-    Event.DeviceModelId = cbi->ProductType;
-    Event.DeviceConnectionType = (cbi->DeviceInterface == DeviceConnectionInterface::Network) ? "Network" : "USB";
+	Event.DeviceUDID = cbi->DeviceUDID;
+	Event.DeviceName = cbi->DeviceName;
+	Event.DeviceType = cbi->ProductType;
+	Event.DeviceOSVersion = cbi->DeviceOSVersion;
+	Event.DeviceModelId = cbi->ProductType;
+	Event.DeviceConnectionType = StringifyDeviceConnection(cbi->DeviceInterface);
 	Event.bIsAuthorized = cbi->IsAuthorized;
-    Event.bCanReboot = false;
-    Event.bCanPowerOn = false;
-    Event.bCanPowerOff = false;
-    FIOSDeviceHelper::OnDeviceConnected().Broadcast(Event);
-
-    // add to the device list
-    ConnectedDevices.Add(Device, Event);
+	Event.bCanReboot = false;
+	Event.bCanPowerOn = false;
+	Event.bCanPowerOff = false;
+	FIOSDeviceHelper::OnDeviceConnected().Broadcast(Event);
+	
+	// add to the device list
+	ConnectedDevices.Add(Device, Event);
 }
 
 void FIOSDeviceHelper::DoDeviceDisconnect(void* CallbackInfo)
@@ -535,12 +669,12 @@ void FIOSDeviceHelper::DoDeviceDisconnect(void* CallbackInfo)
 
 bool FIOSDeviceHelper::InstallIPAOnDevice(const FTargetDeviceId& DeviceId, const FString& IPAPath)
 {
-    return false;
+	return false;
 }
 
 void FIOSDeviceHelper::EnableDeviceCheck(bool OnOff)
 {
-    QueryTask->Enable(OnOff);
+	QueryTask->Enable(OnOff);
 }
 
 #undef LOCTEXT_NAMESPACE
diff --git a/Engine/Source/Developer/IOS/IOSTargetPlatform/Private/IOSTargetDevice.h b/Engine/Source/Developer/IOS/IOSTargetPlatform/Private/IOSTargetDevice.h
--- a/Engine/Source/Developer/IOS/IOSTargetPlatform/Private/IOSTargetDevice.h
+++ b/Engine/Source/Developer/IOS/IOSTargetPlatform/Private/IOSTargetDevice.h
@@ -247,19 +247,23 @@ public:
 	/** Sets the type of the device */
 	void SetDeviceType(const FString InDeviceTypeString)
 	{
-		if (InDeviceTypeString == TEXT("Browser"))
+		if (InDeviceTypeString.Contains(TEXT("Browser")))
 		{
 			DeviceType = ETargetDeviceTypes::Browser;
 		}
-		else if (InDeviceTypeString == TEXT("Console"))
+		else if (InDeviceTypeString.Contains(TEXT("Console")))
 		{
 			DeviceType = ETargetDeviceTypes::Console;
 		}
-		else if (InDeviceTypeString == TEXT("Phone"))
+		else if (InDeviceTypeString.Contains(TEXT("Phone")))
 		{
 			DeviceType = ETargetDeviceTypes::Phone;
 		}
-		else if (InDeviceTypeString == TEXT("Tablet"))
+		else if (InDeviceTypeString.Contains(TEXT("Tablet")))
+		{
+			DeviceType = ETargetDeviceTypes::Tablet;
+		}
+		else if (InDeviceTypeString.Contains(TEXT("iPad")))
 		{
 			DeviceType = ETargetDeviceTypes::Tablet;
 		}
@@ -280,6 +284,10 @@ public:
 		{
 			DeviceConnectionType = ETargetDeviceConnectionTypes::USB;
 		}
+		else if (InDeviceConnectionTypeString == TEXT("Simulator"))
+		{
+			DeviceConnectionType = ETargetDeviceConnectionTypes::Simulator;
+		}
 		else
 		{
 			DeviceConnectionType = ETargetDeviceConnectionTypes::Unknown;
diff --git a/Engine/Source/Developer/IOS/IOSTargetPlatform/Private/IOSTargetPlatform.cpp b/Engine/Source/Developer/IOS/IOSTargetPlatform/Private/IOSTargetPlatform.cpp
--- a/Engine/Source/Developer/IOS/IOSTargetPlatform/Private/IOSTargetPlatform.cpp
+++ b/Engine/Source/Developer/IOS/IOSTargetPlatform/Private/IOSTargetPlatform.cpp
@@ -410,7 +410,7 @@ void FIOSTargetPlatform::HandleDeviceConnected(const FIOSLaunchDaemonPong& Messa
 			Device->SetModelId(Message.DeviceModelId);
 			Device->SetOSVersion(Message.DeviceOSVersion);
 			Device->SetDeviceConnectionType(Message.DeviceConnectionType);
-			Device->SetIsSimulated(Message.DeviceID.Contains(TEXT("Simulator")));
+			Device->SetIsSimulated(Message.DeviceConnectionType.Contains(TEXT("Simulator")));
 
 			ITargetPlatformControls::OnDeviceDiscovered().Broadcast(Device.ToSharedRef());
 		}
diff --git a/Engine/Source/Developer/TargetPlatform/Public/Interfaces/ITargetDevice.h b/Engine/Source/Developer/TargetPlatform/Public/Interfaces/ITargetDevice.h
--- a/Engine/Source/Developer/TargetPlatform/Public/Interfaces/ITargetDevice.h
+++ b/Engine/Source/Developer/TargetPlatform/Public/Interfaces/ITargetDevice.h
@@ -88,7 +88,7 @@ namespace TargetDeviceTypes
 	/**
 	 * Returns the string representation of the specified ETargetDeviceTypes value.
 	 *
-	 * @param Configuration The value to get the string for.
+	 * @param DeviceType The value to get the string for.
 	 * @return A string value.
 	 */
 	inline FString ToString(ETargetDeviceTypes DeviceType)
@@ -116,6 +116,38 @@ namespace TargetDeviceTypes
 	}
 }
 
+namespace TargetDeviceConnectionTypes
+{
+	/**
+	 * Returns the string representation of the specified ETargetDeviceConnectionTypes value.
+	 *
+	 * @param DeviceConnectionType The value to get the string for.
+	 * @return A string value.
+	 */
+	inline FString ToString(ETargetDeviceConnectionTypes DeviceConnectionType)
+	{
+		switch (DeviceConnectionType)
+		{
+		case ETargetDeviceConnectionTypes::USB:
+			return FString("USB");
+
+		case ETargetDeviceConnectionTypes::Wifi:
+			return FString("Wifi");
+
+		case ETargetDeviceConnectionTypes::Ethernet:
+			return FString("Ethernet");
+
+		case ETargetDeviceConnectionTypes::Simulator:
+			return FString("Simulator");
+
+		case ETargetDeviceConnectionTypes::Proprietary:
+			return FString("Proprietary");
+
+		default:
+			return FString("Unknown");
+		}
+	}
+}
 
 /**
  * Enumerates states of threads running on a target device.
