IOS/TVOS/VisionOS TP refactor to TPS/TPC

diff --git a/Engine/Platforms/VisionOS/Source/Developer/VisionOSTargetPlatform/Private/VisionOSTargetPlatformModule.cpp b/Engine/Platforms/VisionOS/Source/Developer/VisionOSTargetPlatform/Private/VisionOSTargetPlatformModule.cpp
--- a/Engine/Platforms/VisionOS/Source/Developer/VisionOSTargetPlatform/Private/VisionOSTargetPlatformModule.cpp
+++ b/Engine/Platforms/VisionOS/Source/Developer/VisionOSTargetPlatform/Private/VisionOSTargetPlatformModule.cpp
@@ -1,9 +1,8 @@
 // Copyright Epic Games, Inc. All Rights Reserved.
 
-#include "IOSTargetPlatform.h"
 #include "Interfaces/ITargetPlatformModule.h"
 #include "Modules/ModuleManager.h"
-
+#include "Common/TargetPlatformBase.h"
 
 /**
  * Module for TVOS as a target platform
@@ -14,14 +13,15 @@ public:
 
 	virtual void GetTargetPlatforms(TArray<ITargetPlatform*>& TargetPlatforms) override
 	{
-		if (FIOSTargetPlatform::IsUsable())
+	}
+
+	virtual void GetTargetPlatforms(TArray<ITargetPlatform*>& TargetPlatforms, TArray<ITargetPlatformSettings*> TargetPlatformSettings, TArray<ITargetPlatformControls*> TargetPlatformControls)
+	{
+		for (ITargetPlatformControls* TargetPlatformControlsIt : TargetPlatformControls)
 		{
-			// add Game and Client TPs
-			TargetPlatforms.Add(new FIOSTargetPlatform(false, true, false));
-			TargetPlatforms.Add(new FIOSTargetPlatform(false, true, true));
+			TargetPlatforms.Add(new FTargetPlatformMerged(TargetPlatformControlsIt->GetTargetPlatformSettings(), TargetPlatformControlsIt));
 		}
 	}
-
 };
 
 
diff --git a/Engine/Platforms/VisionOS/Source/Developer/VisionOSTargetPlatform/VisionOSTargetPlatform.Build.cs b/Engine/Platforms/VisionOS/Source/Developer/VisionOSTargetPlatform/VisionOSTargetPlatform.Build.cs
--- a/Engine/Platforms/VisionOS/Source/Developer/VisionOSTargetPlatform/VisionOSTargetPlatform.Build.cs
+++ b/Engine/Platforms/VisionOS/Source/Developer/VisionOSTargetPlatform/VisionOSTargetPlatform.Build.cs
@@ -13,30 +13,8 @@ public class VisionOSTargetPlatform : ModuleRules
 			new string[]
 			{
 				"Core",
-				"CoreUObject",
 				"TargetPlatform",
 				"DesktopPlatform",
-				"LaunchDaemonMessages",
-				"IOSTargetPlatform",
-				"Projects"
-			}
-		);
-
-		PrivateIncludePathModuleNames.AddRange(
-			new string[] {
-			"Messaging",
-			"TargetDeviceServices",
-		}
-		);
-
-		if (Target.bCompileAgainstEngine)
-		{
-			PrivateDependencyModuleNames.Add("Engine");
-		}
-
-		PrivateIncludePaths.AddRange(
-			new string[] {
-			"Developer/IOS/IOSTargetPlatform/Private"
 			}
 		);
 	}
diff --git a/Engine/Platforms/VisionOS/Source/Developer/VisionOSTargetPlatformControls/Private/VisionOSTargetPlatformControlsModule.cpp b/Engine/Platforms/VisionOS/Source/Developer/VisionOSTargetPlatformControls/Private/VisionOSTargetPlatformControlsModule.cpp
new file mode 100644
--- /dev/null
+++ b/Engine/Platforms/VisionOS/Source/Developer/VisionOSTargetPlatformControls/Private/VisionOSTargetPlatformControlsModule.cpp
@@ -0,0 +1,32 @@
+// Copyright Epic Games, Inc. All Rights Reserved.
+
+#include "IOSTargetPlatformControls.h"
+#include "Interfaces/ITargetPlatformControlsModule.h"
+#include "Interfaces/ITargetPlatformSettingsModule.h"
+#include "Modules/ModuleManager.h"
+
+/**
+ * Module for TVOS as a target platform controls
+ */
+class FVisionOSTargetPlatformControlsModule	: public ITargetPlatformControlsModule
+{
+public:
+	virtual void GetTargetPlatformControls(TArray<ITargetPlatformControls*>& TargetPlatforms, FName& PlatformSettingsModuleName) override
+	{
+		if (FIOSTargetPlatformControls::IsUsable())
+		{
+			ITargetPlatformSettingsModule* ModuleSettings = FModuleManager::GetModulePtr<ITargetPlatformSettingsModule>(PlatformSettingsModuleName);
+			if (ModuleSettings != nullptr)
+			{
+				TArray<ITargetPlatformSettings*> TargetPlatformSettings = ModuleSettings->GetTargetPlatformSettings();
+				check(TargetPlatformSettings.Num() == 1);
+
+				TargetPlatforms.Add(new FIOSTargetPlatformControls(false, true, false, TargetPlatformSettings[0]));
+				TargetPlatforms.Add(new FIOSTargetPlatformControls(false, true, true, TargetPlatformSettings[0]));
+			}
+		}
+	}
+};
+
+
+IMPLEMENT_MODULE(FVisionOSTargetPlatformControlsModule, VisionOSTargetPlatformControls);
diff --git a/Engine/Platforms/VisionOS/Source/Developer/VisionOSTargetPlatformControls/VisionOSTargetPlatformControls.Build.cs b/Engine/Platforms/VisionOS/Source/Developer/VisionOSTargetPlatformControls/VisionOSTargetPlatformControls.Build.cs
new file mode 100644
--- /dev/null
+++ b/Engine/Platforms/VisionOS/Source/Developer/VisionOSTargetPlatformControls/VisionOSTargetPlatformControls.Build.cs
@@ -0,0 +1,46 @@
+// Copyright Epic Games, Inc. All Rights Reserved.
+
+using UnrealBuildTool;
+
+public class VisionOSTargetPlatformControls : ModuleRules
+{
+	public VisionOSTargetPlatformControls(ReadOnlyTargetRules Target) : base(Target)
+	{
+		BinariesSubFolder = "IOS";
+		// We need a short name here since this can run afoul very easily of the `MAX_PATH` limit when
+		// combined with building other targets that make use of this.
+		ShortName = "VisionOSTPCon";
+		SDKVersionRelevantPlatforms.Add(UnrealTargetPlatform.VisionOS);
+
+		PrivateDependencyModuleNames.AddRange(
+			new string[]
+			{
+				"Core",
+				"CoreUObject",
+				"TargetPlatform",
+				"DesktopPlatform",
+				"LaunchDaemonMessages",
+				"IOSTargetPlatformControls",
+				"Projects"
+			}
+		);
+
+		PrivateIncludePathModuleNames.AddRange(
+			new string[] {
+			"Messaging",
+			"TargetDeviceServices",
+		}
+		);
+
+		if (Target.bCompileAgainstEngine)
+		{
+			PrivateDependencyModuleNames.Add("Engine");
+		}
+
+		PrivateIncludePaths.AddRange(
+			new string[] {
+			"Developer/IOS/IOSTargetPlatformControls/Private"
+			}
+		);
+	}
+}
diff --git a/Engine/Platforms/VisionOS/Source/Developer/VisionOSTargetPlatformSettings/Private/VisionOSTargetPlatformSettingsModule.cpp b/Engine/Platforms/VisionOS/Source/Developer/VisionOSTargetPlatformSettings/Private/VisionOSTargetPlatformSettingsModule.cpp
new file mode 100644
--- /dev/null
+++ b/Engine/Platforms/VisionOS/Source/Developer/VisionOSTargetPlatformSettings/Private/VisionOSTargetPlatformSettingsModule.cpp
@@ -0,0 +1,20 @@
+// Copyright Epic Games, Inc. All Rights Reserved.
+
+#include "IOSTargetPlatformSettings.h"
+#include "Interfaces/ITargetPlatformSettingsModule.h"
+#include "Modules/ModuleManager.h"
+
+/**
+ * Module for TVOS as a target platform settings
+ */
+class FVisionOSTargetPlatformSettingsModule	: public ITargetPlatformSettingsModule
+{
+public:
+	virtual void GetTargetPlatformSettings(TArray<ITargetPlatformSettings*>& TargetPlatforms) override
+	{
+		TargetPlatforms.Add(new FIOSTargetPlatformSettings(false, true));
+	}
+};
+
+
+IMPLEMENT_MODULE(FVisionOSTargetPlatformSettingsModule, VisionOSTargetPlatformSettings);
diff --git a/Engine/Platforms/VisionOS/Source/Developer/VisionOSTargetPlatformSettings/VisionOSTargetPlatformSettings.Build.cs b/Engine/Platforms/VisionOS/Source/Developer/VisionOSTargetPlatformSettings/VisionOSTargetPlatformSettings.Build.cs
new file mode 100644
--- /dev/null
+++ b/Engine/Platforms/VisionOS/Source/Developer/VisionOSTargetPlatformSettings/VisionOSTargetPlatformSettings.Build.cs
@@ -0,0 +1,36 @@
+// Copyright Epic Games, Inc. All Rights Reserved.
+
+using UnrealBuildTool;
+
+public class VisionOSTargetPlatformSettings : ModuleRules
+{
+	public VisionOSTargetPlatformSettings(ReadOnlyTargetRules Target) : base(Target)
+	{
+		BinariesSubFolder = "IOS";
+		// We need a short name here since this can run afoul very easily of the `MAX_PATH` limit when
+		// combined with building other targets that make use of this.
+		ShortName = "VisionOSTPSet";
+
+		PrivateDependencyModuleNames.AddRange(
+			new string[]
+			{
+				"Core",
+				"CoreUObject",
+				"TargetPlatform",
+				"DesktopPlatform",
+				"IOSTargetPlatformSettings",
+			}
+		);
+
+		if (Target.bCompileAgainstEngine)
+		{
+			PrivateDependencyModuleNames.Add("Engine");
+		}
+
+		PrivateIncludePaths.AddRange(
+			new string[] {
+			"Developer/IOS/IOSTargetPlatformSettings/Private"
+			}
+		);
+	}
+}
diff --git a/Engine/Platforms/VisionOS/Source/Programs/UnrealBuildTool/UEBuildVisionOS.cs b/Engine/Platforms/VisionOS/Source/Programs/UnrealBuildTool/UEBuildVisionOS.cs
--- a/Engine/Platforms/VisionOS/Source/Programs/UnrealBuildTool/UEBuildVisionOS.cs
+++ b/Engine/Platforms/VisionOS/Source/Programs/UnrealBuildTool/UEBuildVisionOS.cs
@@ -53,8 +53,10 @@ namespace UnrealBuildTool
 		{
 			base.ModifyModuleRulesForOtherPlatform(ModuleName, Rules, Target);
 
+			bool bIsPlatformAvailableForTarget = UEBuildPlatform.IsPlatformAvailableForTarget(Platform, Target, bIgnoreSDKCheck: true);
+			bool bIsPlatformAvailableForTargetWithSDK = UEBuildPlatform.IsPlatformAvailableForTarget(Platform, Target);
 			// don't do any target platform stuff if SDK is not available
-			if (!UEBuildPlatform.IsPlatformAvailableForTarget(Platform, Target))
+			if (!bIsPlatformAvailableForTarget)
 			{
 				return;
 			}
@@ -66,12 +68,17 @@ namespace UnrealBuildTool
 				    if ((ModuleName == "Engine" && Target.bBuildDeveloperTools) ||
 					    (ModuleName == "TargetPlatform" && Target.bForceBuildTargetPlatforms))
 				    {
-					    Rules.DynamicallyLoadedModuleNames.Add("VisionOSTargetPlatform");
+						Rules.DynamicallyLoadedModuleNames.Add("VisionOSTargetPlatformSettings");
+						if (bIsPlatformAvailableForTargetWithSDK)
+						{
+							Rules.DynamicallyLoadedModuleNames.Add("VisionOSTargetPlatform");
+							Rules.DynamicallyLoadedModuleNames.Add("VisionOSTargetPlatformControls");
+						}
 				    }
 			    }
 			}
 
-			if (ModuleName == "UnrealEd")
+			if (ModuleName == "UnrealEd" && bIsPlatformAvailableForTargetWithSDK)
 			{
 				Rules.DynamicallyLoadedModuleNames.Add("VisionOSPlatformEditor");
 			}
diff --git a/Engine/Source/Developer/IOS/IOSTargetPlatform/IOSTargetPlatform.Build.cs b/Engine/Source/Developer/IOS/IOSTargetPlatform/IOSTargetPlatform.Build.cs
--- a/Engine/Source/Developer/IOS/IOSTargetPlatform/IOSTargetPlatform.Build.cs
+++ b/Engine/Source/Developer/IOS/IOSTargetPlatform/IOSTargetPlatform.Build.cs
@@ -13,34 +13,9 @@ public class IOSTargetPlatform : ModuleRules
 			new string[]
 			{
 				"Core",
-				"CoreUObject",
 				"TargetPlatform",
 				"DesktopPlatform",
-				"LaunchDaemonMessages",
-				"Projects",
-				"Json",
-				"AudioPlatformConfiguration",
-				"Sockets",
-				"Networking"
 			}
 		);
-
-
-		PrivateIncludePathModuleNames.AddRange(
-			new string[] {
-				"MessagingCommon",
-				"TargetDeviceServices",
-			}
-		);
-
-		if (Target.bCompileAgainstEngine)
-		{
-			PrivateDependencyModuleNames.Add("Engine");
-		}
-
-		if (Target.Platform == UnrealTargetPlatform.Mac)
-		{
-			PublicAdditionalLibraries.Add("/System/Library/PrivateFrameworks/MobileDevice.framework/Versions/Current/MobileDevice");
-		}
 	}
 }
diff --git a/Engine/Source/Developer/IOS/IOSTargetPlatform/Private/IOSDeviceHelper.cpp b/Engine/Source/Developer/IOS/IOSTargetPlatform/Private/IOSDeviceHelper.cpp
deleted file mode 100644
diff --git a/Engine/Source/Developer/IOS/IOSTargetPlatform/Private/IOSTargetDevice.cpp b/Engine/Source/Developer/IOS/IOSTargetPlatform/Private/IOSTargetDevice.cpp
deleted file mode 100644
diff --git a/Engine/Source/Developer/IOS/IOSTargetPlatform/Private/IOSTargetDevice.h b/Engine/Source/Developer/IOS/IOSTargetPlatform/Private/IOSTargetDevice.h
deleted file mode 100644
diff --git a/Engine/Source/Developer/IOS/IOSTargetPlatform/Private/IOSTargetDeviceOutput.h b/Engine/Source/Developer/IOS/IOSTargetPlatform/Private/IOSTargetDeviceOutput.h
deleted file mode 100644
diff --git a/Engine/Source/Developer/IOS/IOSTargetPlatform/Private/IOSTargetDeviceOutput.inl b/Engine/Source/Developer/IOS/IOSTargetPlatform/Private/IOSTargetDeviceOutput.inl
deleted file mode 100644
diff --git a/Engine/Source/Developer/IOS/IOSTargetPlatform/Private/IOSTargetPlatform.cpp b/Engine/Source/Developer/IOS/IOSTargetPlatform/Private/IOSTargetPlatform.cpp
deleted file mode 100644
diff --git a/Engine/Source/Developer/IOS/IOSTargetPlatform/Private/IOSTargetPlatform.h b/Engine/Source/Developer/IOS/IOSTargetPlatform/Private/IOSTargetPlatform.h
deleted file mode 100644
diff --git a/Engine/Source/Developer/IOS/IOSTargetPlatform/Private/IOSTargetPlatformModule.cpp b/Engine/Source/Developer/IOS/IOSTargetPlatform/Private/IOSTargetPlatformModule.cpp
--- a/Engine/Source/Developer/IOS/IOSTargetPlatform/Private/IOSTargetPlatformModule.cpp
+++ b/Engine/Source/Developer/IOS/IOSTargetPlatform/Private/IOSTargetPlatformModule.cpp
@@ -1,9 +1,8 @@
 // Copyright Epic Games, Inc. All Rights Reserved.
 
-#include "IOSTargetPlatform.h"
 #include "Interfaces/ITargetPlatformModule.h"
 #include "Modules/ModuleManager.h"
-
+#include "Common/TargetPlatformBase.h"
 
 /**
  * Module for iOS as a target platform
@@ -14,13 +13,15 @@ public:
 
 	virtual void GetTargetPlatforms(TArray<ITargetPlatform*>& TargetPlatforms) override
 	{
-		if (FIOSTargetPlatform::IsUsable())
+	}
+
+	virtual void GetTargetPlatforms(TArray<ITargetPlatform*>& TargetPlatforms, TArray<ITargetPlatformSettings*> TargetPlatformSettings, TArray<ITargetPlatformControls*> TargetPlatformControls)
+	{
+		for (ITargetPlatformControls* TargetPlatformControlsIt : TargetPlatformControls)
 		{
-			TargetPlatforms.Add(new FIOSTargetPlatform(false, false, false));
-			TargetPlatforms.Add(new FIOSTargetPlatform(false, false, true));
+			TargetPlatforms.Add(new FTargetPlatformMerged(TargetPlatformControlsIt->GetTargetPlatformSettings(), TargetPlatformControlsIt));
 		}
 	}
-
 };
 
 
diff --git a/Engine/Source/Developer/IOS/IOSTargetPlatformControls/IOSTargetPlatformControls.Build.cs b/Engine/Source/Developer/IOS/IOSTargetPlatformControls/IOSTargetPlatformControls.Build.cs
new file mode 100644
--- /dev/null
+++ b/Engine/Source/Developer/IOS/IOSTargetPlatformControls/IOSTargetPlatformControls.Build.cs
@@ -0,0 +1,49 @@
+// Copyright Epic Games, Inc. All Rights Reserved.
+
+using UnrealBuildTool;
+
+public class IOSTargetPlatformControls : ModuleRules
+{
+	public IOSTargetPlatformControls(ReadOnlyTargetRules Target) : base(Target)
+	{
+		BinariesSubFolder = "IOS";
+		// We need a short name here since this can run afoul very easily of the `MAX_PATH` limit when
+		// combined with building other targets that make use of this.
+		ShortName = "IOSTPCon";
+		SDKVersionRelevantPlatforms.Add(UnrealTargetPlatform.IOS);
+
+		PrivateDependencyModuleNames.AddRange(
+			new string[]
+			{
+				"Core",
+				"CoreUObject",
+				"TargetPlatform",
+				"DesktopPlatform",
+				"LaunchDaemonMessages",
+				"Projects",
+				"Json",
+				"AudioPlatformConfiguration",
+				"Sockets",
+				"Networking",
+				"IOSTargetPlatformSettings"
+			}
+		);
+
+		PrivateIncludePathModuleNames.AddRange(
+			new string[] {
+				"MessagingCommon",
+				"TargetDeviceServices",
+			}
+		);
+
+		if (Target.bCompileAgainstEngine)
+		{
+			PrivateDependencyModuleNames.Add("Engine");
+		}
+
+		if (Target.Platform == UnrealTargetPlatform.Mac)
+		{
+			PublicAdditionalLibraries.Add("/System/Library/PrivateFrameworks/MobileDevice.framework/Versions/Current/MobileDevice");
+		}
+	}
+}
diff --git a/Engine/Source/Developer/IOS/IOSTargetPlatformControls/Private/IOSDeviceHelper.cpp b/Engine/Source/Developer/IOS/IOSTargetPlatformControls/Private/IOSDeviceHelper.cpp
new file mode 100644
--- /dev/null
+++ b/Engine/Source/Developer/IOS/IOSTargetPlatformControls/Private/IOSDeviceHelper.cpp
@@ -0,0 +1,680 @@
+// Copyright Epic Games, Inc. All Rights Reserved.
+
+#include "IOSDeviceHelper.h"
+#include "IOSTargetPlatformControls.h"
+#include "IOSTargetDeviceOutput.h"
+#include "HAL/PlatformProcess.h"
+#include "HAL/Runnable.h"
+#include "HAL/RunnableThread.h"
+#include "Interfaces/ITargetPlatformManagerModule.h"
+#include "Interfaces/ITargetPlatform.h"
+#include "Interfaces/IProjectManager.h"
+#include "Misc/MessageDialog.h"
+#include "Serialization/JsonReader.h"
+#include "Serialization/JsonSerializer.h"
+
+#define LOCTEXT_NAMESPACE "MiscIOSMessages"
+
+DEFINE_LOG_CATEGORY_STATIC(LogIOSDeviceHelper, Log, All);
+
+enum DeviceConnectionInterface
+{
+	NoValue = 0,
+	USB = 1,
+	Network = 2,
+	Simulator = 3,
+	Max = 4
+};
+
+struct FDeviceNotificationCallbackInformation
+{
+	FString DeviceID;
+	FString DeviceName;
+	FString DeviceUDID;
+	FString ProductType;
+	FString DeviceOSVersion;
+	DeviceConnectionInterface DeviceInterface;
+	uint32 msgType;
+	bool IsAuthorized;
+};
+
+
+struct LibIMobileDevice
+{
+	FString DeviceID;
+	FString DeviceName;
+	FString DeviceUDID;
+	FString DeviceType;
+	FString DeviceOSVersion;
+	DeviceConnectionInterface DeviceInterface;
+	bool IsAuthorized;
+	bool IsDealtWith;
+};
+
+const FString StringifyDeviceConnection(DeviceConnectionInterface Interface)
+{
+	switch(Interface)
+	{
+		case DeviceConnectionInterface::NoValue:
+			return TEXT("NoValue");
+
+		case DeviceConnectionInterface::Network:
+			return TEXT("Network");
+
+		case DeviceConnectionInterface::USB:
+			return TEXT("USB");
+
+		case DeviceConnectionInterface::Simulator:
+			return TEXT("Simulator");
+
+		default:
+			UE_LOG(LogIOSDeviceHelper, Warning, TEXT("Unknown DeviceConnectionInterface type:%d"), Interface);
+			return TEXT("NoValue");
+	}	
+}
+
+static TArray<LibIMobileDevice> GetLibIMobileDevices()
+{
+	FString OutStdOut;
+	FString OutStdErr;
+	FString LibimobileDeviceId = GetLibImobileDeviceExe("idevice_id");
+	int ReturnCode;
+	// get the list of devices UDID
+	FPlatformProcess::ExecProcess(*LibimobileDeviceId, TEXT(""), &ReturnCode, &OutStdOut, &OutStdErr, NULL, true);
+	
+	
+	TArray<LibIMobileDevice> ToReturn;
+	// separate out each line
+	TArray<FString> OngoingDeviceIds;
+	
+	OutStdOut.ParseIntoArray(OngoingDeviceIds, TEXT("\n"), true);
+	TArray<FString> DeviceStrings;
+	for (int32 StringIndex = 0; StringIndex < OngoingDeviceIds.Num(); ++StringIndex)
+	{
+		const FString& DeviceUDID = OngoingDeviceIds[StringIndex];
+		DeviceConnectionInterface OngoingDeviceInterface = DeviceConnectionInterface::NoValue;
+		
+		FString OutStdOutInfo;
+		FString OutStdErrInfo;
+		FString LibimobileDeviceInfo = GetLibImobileDeviceExe("ideviceinfo");
+		int ReturnCodeInfo;
+		FString Arguments;
+		
+		if (OngoingDeviceIds[StringIndex].Contains("USB"))
+		{
+			OngoingDeviceInterface = DeviceConnectionInterface::USB;
+		}
+		else if (OngoingDeviceIds[StringIndex].Contains("Network"))
+		{
+			OngoingDeviceInterface = DeviceConnectionInterface::Network;
+		}
+		OngoingDeviceIds[StringIndex].Split(TEXT(" "), &OngoingDeviceIds[StringIndex], nullptr, ESearchCase::CaseSensitive, ESearchDir::FromStart);
+		
+		if (OngoingDeviceInterface == DeviceConnectionInterface::USB)
+		{
+			Arguments = "-u " + DeviceUDID;
+		}
+		else if (OngoingDeviceInterface == DeviceConnectionInterface::Network)
+		{
+			Arguments = "-n -u " + DeviceUDID;
+		}
+		
+		FPlatformProcess::ExecProcess(*LibimobileDeviceInfo, *Arguments, &ReturnCodeInfo, &OutStdOutInfo, &OutStdErrInfo, NULL, true);
+		
+		LibIMobileDevice ToAdd;
+		
+		// ideviceinfo can fail when the connected device is untrusted. It can be "Pairing dialog response pending (-19)", "Invalid HostID (-21)" or "User denied pairing (-18)"
+		// the only thing we can do is to make sure the Trust popup is correctly displayed.
+		if (OutStdErrInfo.Contains("ERROR: "))
+		{
+			if (OutStdErrInfo.Contains("Could not connect to lockdownd"))
+			{
+				// UE_LOG(LogIOSDeviceHelper, Warning, TEXT("Could not pair with connected iOS/tvOS device. Trust this computer by accepting the popup on device."));
+				FString LibimobileDevicePair = GetLibImobileDeviceExe("idevicepair");
+				FString PairArguments = "-u " + DeviceUDID + " pair";
+				FPlatformProcess::ExecProcess(*LibimobileDevicePair, *PairArguments, &ReturnCodeInfo, &OutStdOutInfo, &OutStdErrInfo, NULL, true);
+			}
+			else
+			{
+				UE_LOG(LogIOSDeviceHelper, Warning, TEXT("Libimobile call failed : %s"), *OutStdErrInfo);
+			}
+			OutStdOutInfo.Empty();
+			OutStdErrInfo.Empty();
+			ToAdd.IsAuthorized = false;
+		}
+		else
+		{
+			ToAdd.IsAuthorized = true;
+		}
+		
+		// parse product type and device name
+		FString DeviceName;
+		OutStdOutInfo.Split(TEXT("DeviceName: "), nullptr, &DeviceName, ESearchCase::CaseSensitive, ESearchDir::FromStart);
+		DeviceName.Split(LINE_TERMINATOR, &DeviceName, nullptr, ESearchCase::CaseSensitive, ESearchDir::FromStart);
+		if (!ToAdd.IsAuthorized)
+		{
+			DeviceName = LOCTEXT("IosTvosUnauthorizedDevice", "iOS / tvOS (Unauthorized)").ToString();
+		}
+		else
+		{
+			if (OngoingDeviceInterface == DeviceConnectionInterface::Network)
+			{
+				DeviceName += " [Wifi]";
+			}
+		}
+		
+		FString ProductType;
+		OutStdOutInfo.Split(TEXT("ProductType: "), nullptr, &ProductType, ESearchCase::CaseSensitive, ESearchDir::FromStart);
+		ProductType.Split(LINE_TERMINATOR, &ProductType, nullptr, ESearchCase::CaseSensitive, ESearchDir::FromStart);
+		
+		FString OSVersion; // iOS/iPad OS Version
+		OutStdOutInfo.Split(TEXT("ProductVersion: "), nullptr, &OSVersion, ESearchCase::CaseSensitive, ESearchDir::FromStart);
+		OSVersion.Split(LINE_TERMINATOR, &OSVersion, nullptr, ESearchCase::CaseSensitive, ESearchDir::FromStart);
+		
+		FString DeviceID = FString::Printf(TEXT("%s@%s"),
+										   ProductType.Contains(TEXT("AppleTV")) ? TEXT("TVOS") : TEXT("IOS"),
+										   *DeviceUDID);
+		
+		ToAdd.DeviceID = DeviceID;
+		ToAdd.DeviceUDID = DeviceUDID;
+		ToAdd.DeviceName = DeviceName;
+		ToAdd.DeviceType = ProductType;
+		ToAdd.DeviceOSVersion = OSVersion;
+		ToAdd.DeviceInterface = OngoingDeviceInterface;
+		ToAdd.IsDealtWith = false;
+		ToReturn.Add(ToAdd);
+	}
+	return ToReturn;
+}
+
+static TArray<LibIMobileDevice> GetSimulatorDevices()
+{
+	FString OutStdOut;
+	FString OutStdErr;
+	TArray<LibIMobileDevice> ToReturn;
+
+	int ReturnCode;
+	FString simCommand = TEXT("/usr/bin/xcrun");
+	FString simParams = TEXT("simctl list -je devices available");
+	FPlatformProcess::ExecProcess(*simCommand, *simParams, &ReturnCode, &OutStdOut, &OutStdErr, NULL, true);
+
+	if (OutStdOut.Len() == 0)
+	{
+		UE_LOG(LogIOSDeviceHelper, Warning, TEXT("Unable to access simctl app.  Is Xcode installed?"));
+		return ToReturn;
+	}
+
+	// decode the json object
+	TSharedPtr<FJsonObject> JsonObject;
+	TSharedRef<TJsonReader<>> Reader = TJsonReaderFactory<>::Create(*OutStdOut);//ANSI_TO_TCHAR((const ANSICHAR *)Msg.GetData()));
+	if (!FJsonSerializer::Deserialize(Reader, JsonObject) || !JsonObject.IsValid())
+	{
+		UE_LOG(LogIOSDeviceHelper, Error, TEXT("Failed to parse message from SimCtl utility: %s"), *Reader->GetErrorMessage());
+		return ToReturn;
+	}
+
+	const TSharedPtr<FJsonObject>* SimDevices;
+	if (JsonObject->TryGetObjectField(TEXT("devices"), /*out*/ SimDevices))
+	{
+		// This should be a list of iOS/tvOS/xrOS versions, each with an array of simulators for that version
+		for (const auto& OSSimPair : (*SimDevices)->Values)
+		{
+			// We only track iOS, for now
+			// The Key is format is: com.apple.CoreSimulator.SimRuntime.iOS-17-2
+			int IOSIndex = OSSimPair.Key.Find(TEXT("iOS"));
+			if (IOSIndex > 0)
+			{
+				FString OSVer = OSSimPair.Key.RightChop(IOSIndex+4);
+				OSVer.ReplaceInline(TEXT("-"), TEXT("."));
+
+				LibIMobileDevice ToAdd;
+				ToAdd.IsAuthorized = true;
+				ToAdd.DeviceOSVersion = OSVer;
+
+				const TArray<TSharedPtr<FJsonValue> >* SimDataArray;
+				if (OSSimPair.Value->TryGetArray(SimDataArray))
+				{
+					for (const auto& SimData : *SimDataArray)
+					{
+						const TSharedPtr<FJsonObject>* SimDataObj;
+						if (SimData->TryGetObject(SimDataObj))
+						{
+							// For now, we only care about booted devices
+							const FString DeviceState = (*SimDataObj)->GetStringField(TEXT("state"));
+							if (DeviceState == TEXT("Booted"))
+							{
+								const FString DeviceName = (*SimDataObj)->GetStringField(TEXT("name"));
+								ToAdd.DeviceUDID = (*SimDataObj)->GetStringField(TEXT("udid"));
+								ToAdd.DeviceType = (*SimDataObj)->GetStringField(TEXT("name"));
+								ToAdd.DeviceName = DeviceName;
+							
+								FString DeviceID = FString::Printf(TEXT("IOS@%s"), *ToAdd.DeviceUDID);
+								ToAdd.DeviceID = DeviceID;
+								ToAdd.DeviceInterface = DeviceConnectionInterface::Simulator;
+							
+								ToReturn.Add(ToAdd);
+							}
+						}
+					}
+				}
+			}
+			
+		}
+
+	}
+	
+	return ToReturn;
+}
+
+class FIOSDevice
+{
+public:
+	FIOSDevice(FString InID, FString InName, DeviceConnectionInterface InConnectionType)
+	: UDID(InID)
+	, Name(InName)
+	, ConnectionType(InConnectionType)
+	{
+	}
+	
+	~FIOSDevice()
+	{
+	}
+	
+	FString SerialNumber() const
+	{
+		return UDID;
+	}
+	
+	DeviceConnectionInterface ConnectionInterface() const
+	{
+		return ConnectionType;
+	}
+	
+private:
+	FString UDID;
+	FString Name;
+	DeviceConnectionInterface ConnectionType;
+};
+
+/**
+ * Delegate type for devices being connected or disconnected from the machine
+ *
+ * The first parameter is newly added or removed device
+ */
+DECLARE_MULTICAST_DELEGATE_OneParam(FDeviceNotification, void*)
+
+// recheck once per minute
+#define        RECHECK_COUNTER_RESET            12
+
+class FDeviceQueryTask
+: public FRunnable
+{
+public:
+	FDeviceQueryTask()
+	: Stopping(false)
+	, bCheckDevices(true)
+	, NeedSDKCheck(true)
+	, RetryQuery(5)
+	{}
+	
+	virtual bool Init() override
+	{
+		return true;
+	}
+	
+	virtual uint32 Run() override
+	{
+		while (!Stopping)
+		{
+			if (IsEngineExitRequested())
+			{
+				break;
+			}
+			if (GetTargetPlatformManager())
+			{
+				FString OutTutorialPath;
+				const ITargetPlatform* Platform = GetTargetPlatformManager()->FindTargetPlatform(TEXT("IOS"));
+				if (Platform)
+				{
+					if (Platform->IsSdkInstalled(false, OutTutorialPath))
+					{
+						break;
+					}
+				}
+				Enable(false);
+				return 0;
+			}
+			else
+			{
+				FPlatformProcess::Sleep(1.0f);
+			}
+		}
+		int RecheckCounter = RECHECK_COUNTER_RESET;
+		while (!Stopping)
+		{
+			if (IsEngineExitRequested())
+			{
+				break;
+			}
+			if (bCheckDevices)
+			{
+#if WITH_EDITOR
+				if (!IsRunningCommandlet())
+				{
+					//if (NeedSDKCheck)
+					//{
+					//    NeedSDKCheck = false;
+					//    FProjectStatus ProjectStatus;
+					//    if (!IProjectManager::Get().QueryStatusForCurrentProject(ProjectStatus) || (!ProjectStatus.IsTargetPlatformSupported(TEXT("IOS")) && !ProjectStatus.IsTargetPlatformSupported(TEXT("TVOS"))))
+					//    {
+					//        Enable(false);
+					//    }
+					//}
+					//else
+					{
+						// BHP - Turning off device check to prevent it from interfering with packaging
+						QueryDevices();
+					}
+				}
+#else
+				QueryDevices();
+#endif
+			}
+			RecheckCounter--;
+			if (RecheckCounter < 0)
+			{
+				RecheckCounter = RECHECK_COUNTER_RESET;
+				bCheckDevices = true;
+				NeedSDKCheck = true;
+			}
+			
+			FPlatformProcess::Sleep(5.0f);
+		}
+		
+		return 0;
+	}
+	
+	virtual void Stop() override
+	{
+		Stopping = true;
+	}
+	
+	virtual void Exit() override
+	{}
+	
+	FDeviceNotification& OnDeviceNotification()
+	{
+		return DeviceNotification;
+	}
+	
+	void Enable(bool bInCheckDevices)
+	{
+		bCheckDevices = bInCheckDevices;
+	}
+	
+private:
+	
+	void NotifyDeviceChange(LibIMobileDevice& Device, bool bAdd)
+	{
+		FDeviceNotificationCallbackInformation CallbackInfo;
+		
+		if (bAdd)
+		{
+			CallbackInfo.DeviceID = Device.DeviceID;
+			CallbackInfo.DeviceName = Device.DeviceName;
+			CallbackInfo.DeviceUDID = Device.DeviceUDID;
+			CallbackInfo.DeviceInterface = Device.DeviceInterface;
+			CallbackInfo.ProductType = Device.DeviceType;
+			CallbackInfo.DeviceOSVersion = Device.DeviceOSVersion;
+			CallbackInfo.msgType = 1;
+			CallbackInfo.IsAuthorized = Device.IsAuthorized;
+		}
+		else
+		{
+			CallbackInfo.DeviceID = Device.DeviceID;
+			CallbackInfo.DeviceUDID = Device.DeviceUDID;
+			CallbackInfo.DeviceInterface = Device.DeviceInterface;
+			CallbackInfo.msgType = 2;
+			DeviceNotification.Broadcast(&CallbackInfo);
+			
+		}
+		DeviceNotification.Broadcast(&CallbackInfo);
+	}
+	
+	void QueryDevices()
+	{
+		TArray<LibIMobileDevice> SimulatorDevices;
+		bool HasSimDevices = false;
+		bool HasDevices = true;
+
+		bool bEnableSimulatorSupport = false;
+		GConfig->GetBool(TEXT("/Script/IOSRuntimeSettings.IOSRuntimeSettings"), TEXT("bEnableSimulatorSupport"), bEnableSimulatorSupport, GEngineIni);
+		if (bEnableSimulatorSupport)
+		{
+			SimulatorDevices = GetSimulatorDevices();
+			HasSimDevices = SimulatorDevices.Num() > 0;
+		}
+
+		FString LibimobileDeviceId = GetLibImobileDeviceExe("idevice_id");
+		if (LibimobileDeviceId.Len() == 0)
+		{
+			UE_LOG(LogIOSDeviceHelper, Log, TEXT("idevice_id (iOS device detection) executable missing. Turning off iOS/tvOS device detection."));
+			HasDevices = false;
+			Enable(false);
+			if (!HasSimDevices)
+			{
+				return;
+			}
+		}
+
+		// get the list of devices UDID
+		if (HasDevices)
+		{
+			FString OutStdOut;
+			FString OutStdErr;
+			int ReturnCode;
+			FPlatformProcess::ExecProcess(*LibimobileDeviceId, TEXT(""), &ReturnCode, &OutStdOut, &OutStdErr, NULL, true);
+			if (OutStdOut.Len() == 0)
+			{
+				RetryQuery--;
+				if (RetryQuery < 0)
+				{
+					UE_LOG(LogIOSDeviceHelper, Verbose, TEXT("IOS device listing is disabled for 1 minute (too many failed attempts)!"));
+					Enable(false);
+				}
+				for (LibIMobileDevice device : CachedDevices)
+				{
+					NotifyDeviceChange(device, false);
+				}
+				CachedDevices.Empty();
+				HasDevices = false;
+				if (!HasSimDevices)
+				{
+					return;
+				}
+			}
+			RetryQuery = 5;
+		}
+
+		TArray<LibIMobileDevice> ParsedDevices;
+		if (HasDevices)
+		{
+			ParsedDevices = GetLibIMobileDevices();
+		}
+
+		ParsedDevices.Append(SimulatorDevices);
+		
+		for (int32 Index = 0; Index < ParsedDevices.Num(); ++Index)
+		{
+			LibIMobileDevice *Found = CachedDevices.FindByPredicate(
+				[&](LibIMobileDevice Element)
+				{
+					return (Element.DeviceUDID == ParsedDevices[Index].DeviceUDID &&
+							Element.DeviceInterface == ParsedDevices[Index].DeviceInterface);
+				});
+			if (Found != nullptr)
+			{
+				if (Found->IsAuthorized != ParsedDevices[Index].IsAuthorized)
+				{
+					NotifyDeviceChange(ParsedDevices[Index], false);
+					NotifyDeviceChange(ParsedDevices[Index], true);
+				}
+				Found->IsDealtWith = true;
+			}
+			else
+			{
+				NotifyDeviceChange(ParsedDevices[Index], true);
+			}
+			
+		}
+		
+		for (int32 Index = 0; Index < CachedDevices.Num(); ++Index)
+		{
+			if (!CachedDevices[Index].IsDealtWith)
+			{
+				NotifyDeviceChange(CachedDevices[Index], false);
+			}
+		}
+		
+		CachedDevices.Empty();
+		CachedDevices = ParsedDevices;
+	}
+	
+	bool Stopping;
+	bool bCheckDevices;
+	bool NeedSDKCheck;
+	int RetryQuery;
+	TArray<LibIMobileDevice> CachedDevices;
+	FDeviceNotification DeviceNotification;
+};
+
+/* FIOSDeviceHelper structors
+ *****************************************************************************/
+static TMap<FIOSDevice*, FIOSLaunchDaemonPong> ConnectedDevices;
+static FDeviceQueryTask* QueryTask = NULL;
+static FRunnableThread* QueryThread = NULL;
+static TArray<FDeviceNotificationCallbackInformation> NotificationMessages;
+static FTickerDelegate TickDelegate;
+
+bool FIOSDeviceHelper::MessageTickDelegate(float DeltaTime)
+{
+	QUICK_SCOPE_CYCLE_COUNTER(STAT_FIOSDeviceHelper_MessageTickDelegate);
+	
+	for (int Index = 0; Index < NotificationMessages.Num(); ++Index)
+	{
+		FDeviceNotificationCallbackInformation cbi = NotificationMessages[Index];
+		FIOSDeviceHelper::DeviceCallback(&cbi);
+	}
+	NotificationMessages.Empty();
+	
+	return true;
+}
+
+void FIOSDeviceHelper::Initialize(bool bIsTVOS)
+{
+	if(!bIsTVOS)
+	{
+		// add the message pump
+		TickDelegate = FTickerDelegate::CreateStatic(MessageTickDelegate);
+		FTSTicker::GetCoreTicker().AddTicker(TickDelegate, 5.0f);
+		
+		// kick off a thread to query for connected devices
+		QueryTask = new FDeviceQueryTask();
+		QueryTask->OnDeviceNotification().AddStatic(FIOSDeviceHelper::DeviceCallback);
+		
+		static int32 QueryTaskCount = 1;
+		if (QueryTaskCount == 1)
+		{
+			// create the socket subsystem (loadmodule in game thread)
+			ISocketSubsystem* SSS = ISocketSubsystem::Get();
+			QueryThread = FRunnableThread::Create(QueryTask, *FString::Printf(TEXT("FIOSDeviceHelper.QueryTask_%d"), QueryTaskCount++), 128 * 1024, TPri_Normal);
+		}
+	}
+}
+
+void FIOSDeviceHelper::DeviceCallback(void* CallbackInfo)
+{
+	struct FDeviceNotificationCallbackInformation* cbi = (FDeviceNotificationCallbackInformation*)CallbackInfo;
+	
+	if (!IsInGameThread())
+	{
+		NotificationMessages.Add(*cbi);
+	}
+	else
+	{
+		switch(cbi->msgType)
+		{
+			case 1:
+				FIOSDeviceHelper::DoDeviceConnect(CallbackInfo);
+				break;
+				
+			case 2:
+				FIOSDeviceHelper::DoDeviceDisconnect(CallbackInfo);
+				break;
+		}
+	}
+}
+
+void FIOSDeviceHelper::DoDeviceConnect(void* CallbackInfo)
+{
+	// connect to the device
+	struct FDeviceNotificationCallbackInformation* cbi = (FDeviceNotificationCallbackInformation*)CallbackInfo;
+	FIOSDevice* Device = new FIOSDevice(cbi->DeviceUDID, cbi->DeviceName, cbi->DeviceInterface);
+	
+	// fire the event
+	FIOSLaunchDaemonPong Event;
+	Event.DeviceID = cbi->DeviceID;
+	Event.DeviceUDID = cbi->DeviceUDID;
+	Event.DeviceName = cbi->DeviceName;
+	Event.DeviceType = cbi->ProductType;
+	Event.DeviceOSVersion = cbi->DeviceOSVersion;
+	Event.DeviceModelId = cbi->ProductType;
+	Event.DeviceConnectionType = StringifyDeviceConnection(cbi->DeviceInterface);
+	Event.bIsAuthorized = cbi->IsAuthorized;
+	Event.bCanReboot = false;
+	Event.bCanPowerOn = false;
+	Event.bCanPowerOff = false;
+	FIOSDeviceHelper::OnDeviceConnected().Broadcast(Event);
+	
+	// add to the device list
+	ConnectedDevices.Add(Device, Event);
+}
+
+void FIOSDeviceHelper::DoDeviceDisconnect(void* CallbackInfo)
+{
+	struct FDeviceNotificationCallbackInformation* cbi = (FDeviceNotificationCallbackInformation*)CallbackInfo;
+	FIOSDevice* device = NULL;
+	for (auto DeviceIterator = ConnectedDevices.CreateIterator(); DeviceIterator; ++DeviceIterator)
+	{
+		if (DeviceIterator.Key()->SerialNumber() == cbi->DeviceUDID &&
+			DeviceIterator.Key()->ConnectionInterface() == cbi->DeviceInterface)
+		{
+			device = DeviceIterator.Key();
+			break;
+		}
+	}
+	
+	if (device != NULL)
+	{
+		// extract the device id from the connected list
+		FIOSLaunchDaemonPong Event = ConnectedDevices.FindAndRemoveChecked(device);
+		
+		// fire the event
+		FIOSDeviceHelper::OnDeviceDisconnected().Broadcast(Event);
+		
+		// delete the device
+		delete device;
+	}
+}
+
+bool FIOSDeviceHelper::InstallIPAOnDevice(const FTargetDeviceId& DeviceId, const FString& IPAPath)
+{
+	return false;
+}
+
+void FIOSDeviceHelper::EnableDeviceCheck(bool OnOff)
+{
+	QueryTask->Enable(OnOff);
+}
+
+#undef LOCTEXT_NAMESPACE
diff --git a/Engine/Source/Developer/IOS/IOSTargetPlatformControls/Private/IOSDeviceHelper.h b/Engine/Source/Developer/IOS/IOSTargetPlatformControls/Private/IOSDeviceHelper.h
new file mode 100644
--- /dev/null
+++ b/Engine/Source/Developer/IOS/IOSTargetPlatformControls/Private/IOSDeviceHelper.h
@@ -0,0 +1,69 @@
+// Copyright Epic Games, Inc. All Rights Reserved.
+
+#pragma once
+
+#include "CoreMinimal.h"
+
+/**
+ * Type definition for shared pointers to instances of FIOSDeviceHelper
+ */
+typedef TSharedPtr<class FIOSDeviceHelper, ESPMode::ThreadSafe> FIOSDeviceHelperPtr;
+
+/**
+ * Type definition for shared references to instances of FIOSDeviceHelper
+ */
+typedef TSharedRef<class FIOSDeviceHelper, ESPMode::ThreadSafe> FIOSDeviceHelperRef;
+
+/**
+ * Delegate type for devices being connected or disconnected from the machine
+ *
+ * The first parameter is newly added or removed device
+ */
+DECLARE_MULTICAST_DELEGATE_OneParam(FOnDeviceConnectEvent, const struct FIOSLaunchDaemonPong&);
+
+/**
+ * Implemented the iOS device helper class
+ */
+class FIOSDeviceHelper
+{
+public:
+    static void Initialize(bool bIsTVOS);
+    
+    /**
+     * Returns a delegate that is executed when a device is connected
+     *
+     * @return the delegate
+     */
+    static FOnDeviceConnectEvent& OnDeviceConnected()
+    {
+        static FOnDeviceConnectEvent DeviceConnectedDelegate;
+        return DeviceConnectedDelegate;
+    }
+    
+    /**
+     * Returns a delegate the is executed when a device is disconnected
+     *
+     * @return the delegate
+     */
+    static FOnDeviceConnectEvent& OnDeviceDisconnected()
+    {
+        static FOnDeviceConnectEvent DeviceDisconnectedDelegate;
+        return DeviceDisconnectedDelegate;
+    }
+    
+    /**
+     * Installs an ipa on to the device
+     */
+    static bool InstallIPAOnDevice(const class FTargetDeviceId& DeviceId, const FString& IPAPath);
+    
+	/**
+	 * Suspends/Enables the device connect/disconnect thread
+	 */
+	static void EnableDeviceCheck(bool OnOff);
+
+private:
+    static void DeviceCallback(void*);
+    static void DoDeviceConnect(void*);
+    static void DoDeviceDisconnect(void*);
+	static bool MessageTickDelegate(float);
+};
diff --git a/Engine/Source/Developer/IOS/IOSTargetPlatformControls/Private/IOSTargetDevice.cpp b/Engine/Source/Developer/IOS/IOSTargetPlatformControls/Private/IOSTargetDevice.cpp
new file mode 100644
--- /dev/null
+++ b/Engine/Source/Developer/IOS/IOSTargetPlatformControls/Private/IOSTargetDevice.cpp
@@ -0,0 +1,195 @@
+// Copyright Epic Games, Inc. All Rights Reserved.
+
+#include "IOSTargetDevice.h"
+
+#include "HAL/PlatformProcess.h"
+#include "IOSMessageProtocol.h"
+#include "Interfaces/ITargetPlatformControls.h"
+#include "Async/Async.h"
+#include "MessageEndpoint.h"
+#include "MessageEndpointBuilder.h"
+#include "IOSTargetDeviceOutput.h"
+#if PLATFORM_WINDOWS
+// for windows mutex
+#include "Windows/AllowWindowsPlatformTypes.h"
+#endif // #if PLATFORM_WINDOWS
+
+FIOSTargetDevice::FIOSTargetDevice(const ITargetPlatformControls& InTargetPlatformControls)
+	: TargetPlatformControls(InTargetPlatformControls)
+	, DeviceEndpoint()
+	, AppId()
+	, bCanReboot(false)
+	, bCanPowerOn(false)
+	, bCanPowerOff(false)
+	, DeviceType(ETargetDeviceTypes::Indeterminate)
+	, DeviceModelId(TEXT(""))
+{
+	DeviceId = FTargetDeviceId(TargetPlatformControls.PlatformName(), FPlatformProcess::ComputerName());
+	DeviceName = FPlatformProcess::ComputerName();
+	MessageEndpoint = FMessageEndpoint::Builder("FIOSTargetDevice").Build();
+	DeviceConnectionType = ETargetDeviceConnectionTypes::USB;
+}
+
+
+bool FIOSTargetDevice::Connect()
+{
+	// @todo zombie - Probably need to write a specific ConnectTo(IpAddr) function for setting up a RemoteEndpoint for talking to the Daemon
+	// Returning true since, if this exists, a device exists.
+
+	return true;
+}
+
+void FIOSTargetDevice::Disconnect()
+{
+}
+
+int32 FIOSTargetDevice::GetProcessSnapshot(TArray<FTargetDeviceProcessInfo>& OutProcessInfos)
+{
+	return 0;
+}
+
+ETargetDeviceTypes FIOSTargetDevice::GetDeviceType() const
+{
+	return DeviceType;
+}
+
+ETargetDeviceConnectionTypes FIOSTargetDevice::GetDeviceConnectionType() const
+{
+	return DeviceConnectionType;
+}
+
+FTargetDeviceId FIOSTargetDevice::GetId() const
+{
+	return DeviceId;
+}
+
+FString FIOSTargetDevice::GetName() const
+{
+	return DeviceName;
+}
+
+FString FIOSTargetDevice::GetOperatingSystemName()
+{
+	return TargetPlatformControls.PlatformName();
+}
+
+FString FIOSTargetDevice::GetModelId() const
+{
+	return DeviceModelId;
+}
+
+FString FIOSTargetDevice::GetOSVersion() const
+{
+	return DeviceOSVersion;
+}
+
+const class ITargetPlatformSettings& FIOSTargetDevice::GetPlatformSettings() const
+{
+	return *(TargetPlatformControls.GetTargetPlatformSettings());
+}
+const class ITargetPlatformControls& FIOSTargetDevice::GetPlatformControls() const
+{
+	return TargetPlatformControls;
+}
+
+bool FIOSTargetDevice::IsConnected()
+{
+	return true;
+}
+
+bool FIOSTargetDevice::IsDefault() const
+{
+	return true;
+}
+
+bool FIOSTargetDevice::PowerOff(bool Force)
+{
+	// @todo zombie - Supported by the Daemon?
+
+	return false;
+}
+
+bool FIOSTargetDevice::PowerOn()
+{
+	// @todo zombie - Supported by the Daemon?
+
+	return false;
+}
+
+bool FIOSTargetDevice::Reboot(bool bReconnect)
+{
+	// @todo zombie - Supported by the Daemon?
+
+	return false;
+}
+
+bool FIOSTargetDevice::SupportsFeature(ETargetDeviceFeatures Feature) const
+{
+	switch (Feature)
+	{
+	case ETargetDeviceFeatures::Reboot:
+		return bCanReboot;
+
+	case ETargetDeviceFeatures::PowerOn:
+		return bCanPowerOn;
+
+	case ETargetDeviceFeatures::PowerOff:
+		return bCanPowerOff;
+
+	case ETargetDeviceFeatures::ProcessSnapshot:
+		return false;
+
+	default:
+		return false;
+	}
+}
+
+bool FIOSTargetDevice::TerminateProcess(const int64 ProcessId)
+{
+	return false;
+}
+
+void FIOSTargetDevice::SetUserCredentials(const FString& UserName, const FString& UserPassword)
+{
+}
+
+bool FIOSTargetDevice::GetUserCredentials(FString& OutUserName, FString& OutUserPassword)
+{
+	return false;
+}
+
+inline void FIOSTargetDevice::ExecuteConsoleCommand(const FString& ExecCommand) const
+{
+	FString OutStdOut;
+	FString OutStdErr;
+	FString Exe = GetLibImobileDeviceExe("itcpconnect");
+	FString Params = FString::Printf(TEXT(" -u %s 8888"), *DeviceId.GetDeviceName());
+
+	void* StdInPipe_Read = nullptr;
+	void* StdInPipe_Write = nullptr;
+
+	verify(FPlatformProcess::CreatePipe(StdInPipe_Read, StdInPipe_Write, true));
+
+	FProcHandle ProcHandle = FPlatformProcess::CreateProc(*Exe, *Params, true, false, false, nullptr, 0, nullptr, nullptr, StdInPipe_Read);
+	if (ProcHandle.IsValid())
+	{
+		FPlatformProcess::WritePipe(StdInPipe_Write, ExecCommand, nullptr);
+		FPlatformProcess::ClosePipe(StdInPipe_Read, StdInPipe_Write);
+		FPlatformProcess::WaitForProc(ProcHandle);
+	}
+	else
+	{
+		FPlatformProcess::ClosePipe(StdInPipe_Read, StdInPipe_Write);
+	}
+}
+
+inline ITargetDeviceOutputPtr FIOSTargetDevice::CreateDeviceOutputRouter(FOutputDevice* Output) const
+{
+	FIOSTargetDeviceOutputPtr DeviceOutputPtr = MakeShareable(new FIOSTargetDeviceOutput());
+	if (DeviceOutputPtr->Init(*this, Output))
+	{
+		return DeviceOutputPtr;
+	}
+
+	return nullptr;
+}
\ No newline at end of file
diff --git a/Engine/Source/Developer/IOS/IOSTargetPlatformControls/Private/IOSTargetDevice.h b/Engine/Source/Developer/IOS/IOSTargetPlatformControls/Private/IOSTargetDevice.h
new file mode 100644
--- /dev/null
+++ b/Engine/Source/Developer/IOS/IOSTargetPlatformControls/Private/IOSTargetDevice.h
@@ -0,0 +1,317 @@
+// Copyright Epic Games, Inc. All Rights Reserved.
+
+#pragma once
+
+#include "CoreMinimal.h"
+#include "IMessageContext.h"
+#include "Interfaces/ITargetDevice.h"
+#include "HAL/Runnable.h"
+#include "HAL/RunnableThread.h"
+#include "Containers/Queue.h"
+
+class FMessageEndpoint;
+class ITargetPlatformControls;
+
+
+/** Type definition for shared pointers to instances of FIOSTargetDevice. */
+typedef TSharedPtr<class FIOSTargetDevice, ESPMode::ThreadSafe> FIOSTargetDevicePtr;
+
+/** Type definition for shared references to instances of FIOSTargetDevice. */
+typedef TSharedRef<class FIOSTargetDevice, ESPMode::ThreadSafe> FIOSTargetDeviceRef;
+
+/** Type definition for shared references to instances of FIOSTargetDeviceOutput. */
+typedef TSharedPtr<class FIOSTargetDeviceOutput, ESPMode::ThreadSafe> FIOSTargetDeviceOutputPtr;
+
+/**
+ * Handles the communication to the Deployment Server over TCP (will start the DeploymentServer if no instance is found running)
+ */
+class FTcpDSCommander : FRunnable
+{
+public:
+    
+    /**
+     * Creates and initializes a new instance.
+     *
+     */
+    FTcpDSCommander(const uint8* Data, int32 Count, TQueue<FString>& InOutputQueue);
+    
+    /** Virtual destructor. */
+    virtual ~FTcpDSCommander();
+
+	/** Check if DeploymentServer mutex is active*/
+	static bool IsDSRunning();
+    
+public:
+    
+    //~ FRunnable interface
+    
+    virtual void Exit() override;
+    virtual bool Init() override;
+    virtual uint32 Run() override;
+    virtual void Stop() override;
+    
+    inline bool IsValid()
+    {
+        return (Thread != nullptr);
+    }
+    inline bool IsStopped()
+    {
+        return bStopped;
+    }
+    inline bool WasSuccess()
+    {
+        return bIsSuccess;
+    }
+	inline bool IsSystemError()
+	{
+		return bIsSystemError;
+	}
+    
+private:
+    
+    inline bool StartDSProcess();
+    
+    /** For the thread */
+    bool bStopping;
+    bool bStopped;
+    bool bIsSuccess;
+	bool bIsSystemError; ///< Deployment server was not able to start, or connection to it could not be made
+    
+    /** */
+    class FSocket* DSSocket;
+    
+    /** Holds the thread object. */
+    FRunnableThread* Thread;
+    
+	TQueue<FString>& OutputQueue;
+    uint8* DSCommand;
+    int32 DSCommandLen;
+    double LastActivity;
+};
+
+/**
+ * Implements an iOS target device.
+ */
+class FIOSTargetDevice
+	: public ITargetDevice
+{
+public:
+
+	/**
+	 * Create and initialize a new instance.
+	 *
+	 * @param InTargetPlatformControls The target platform that owns the device.
+	 */
+	FIOSTargetDevice(const ITargetPlatformControls& InTargetPlatformControls);
+
+public:
+
+	//~ ITargetDevice interface
+
+	virtual bool Connect() override;
+	virtual void Disconnect() override;
+	virtual int32 GetProcessSnapshot(TArray<FTargetDeviceProcessInfo>& OutProcessInfos) override;
+	virtual ETargetDeviceTypes GetDeviceType() const override;
+	virtual ETargetDeviceConnectionTypes GetDeviceConnectionType() const override;
+	virtual FTargetDeviceId GetId() const override;
+	virtual FString GetName() const override;
+	virtual FString GetOperatingSystemName() override;
+	virtual FString GetModelId() const override;
+	virtual FString GetOSVersion() const override;
+	virtual const class ITargetPlatformSettings& GetPlatformSettings() const override;
+	virtual const class ITargetPlatformControls& GetPlatformControls() const override;
+	virtual bool IsConnected() override;
+	virtual bool IsDefault() const override;
+	virtual bool PowerOff(bool Force) override;
+	virtual bool PowerOn() override;
+	virtual bool IsAuthorized() const override { return bIsDeviceAuthorized; }
+	virtual bool Reboot(bool bReconnect = false) override;
+	virtual bool SupportsFeature(ETargetDeviceFeatures Feature) const;
+	virtual bool TerminateProcess(const int64 ProcessId) override;
+	virtual void SetUserCredentials(const FString& UserName, const FString& UserPassword) override;
+	virtual bool GetUserCredentials(FString& OutUserName, FString& OutUserPassword) override;
+	virtual void ExecuteConsoleCommand(const FString& ExecCommand) const override;
+	virtual ITargetDeviceOutputPtr CreateDeviceOutputRouter(FOutputDevice* Output) const override;
+
+public:
+
+	/** Timeout check for removing stale devices */
+	FDateTime LastPinged;
+
+private:
+
+	/** The current status of this device. */
+//	ETargetDeviceStatus::Type Status;
+
+/** Holds a reference to the device's target platform controls. */
+	const ITargetPlatformControls& TargetPlatformControls;
+
+	/** Contains the address of the remote device */
+	FMessageAddress DeviceEndpoint;
+
+	/** MessageEndpoint for communicating with remote device */
+	TSharedPtr<FMessageEndpoint, ESPMode::ThreadSafe> MessageEndpoint;
+
+	/** Contains the current AppID/GameName for Deployment/launching. */
+	FString AppId;
+
+	/** Contains the build configuration of the app to deploy */
+	EBuildConfiguration BuildConfiguration;
+
+	/** Lets us know whether the thing is a sim device or a physical device. */
+	bool bIsSimulated;
+
+private:
+
+	/** Remote rebootable */
+	bool bCanReboot;
+
+	/** Remote bootable */
+	bool bCanPowerOn;
+
+	/** Remote shutdown-able */
+	bool bCanPowerOff;
+
+	/** Id of device */
+	FTargetDeviceId DeviceId;
+
+	/** Name of device */
+	FString DeviceName;
+
+	// Holds a flag indicating whether the device is USB / OTA comms authorized
+	bool bIsDeviceAuthorized;
+
+	/** Type of device */
+	ETargetDeviceTypes DeviceType;
+
+	/** The specific model identifier of the device */
+	FString DeviceModelId;
+
+	/** The iOS/tvOS/iPadOS OS version */
+	FString DeviceOSVersion;
+
+	/** Type of device connection (USB or Wifi) */
+	ETargetDeviceConnectionTypes DeviceConnectionType;
+
+public:
+
+	void SetFeature(ETargetDeviceFeatures InFeature, bool bFlag)
+	{
+		if (InFeature == ETargetDeviceFeatures::Reboot)
+		{
+			bCanReboot = bFlag;
+		}
+		else if (InFeature == ETargetDeviceFeatures::PowerOn)
+		{
+			bCanPowerOn = bFlag;
+		}
+		else if (InFeature == ETargetDeviceFeatures::PowerOff)
+		{
+			bCanPowerOff = bFlag;
+		}
+	}
+
+	/** Sets device id */
+	void SetDeviceId(const FTargetDeviceId InDeviceId)
+	{
+		DeviceId = InDeviceId;
+	}
+
+	/** Sets the name of the device */
+	void SetDeviceName(const FString InDeviceName)
+	{
+		DeviceName = InDeviceName;
+	}
+
+	/** Sets the modelId of the device */
+	void SetModelId(const FString InModelId)
+	{
+		DeviceModelId = InModelId;
+	}
+
+	/** Sets the OS version of the device */
+	void SetOSVersion(const FString InOSVersion)
+	{
+		DeviceOSVersion = InOSVersion;
+	}
+
+	/**
+	 * Sets the device's authorization state.
+	 *
+	 * @param bInIsAuthorized - Whether the device is authorized for USB communications.
+	 */
+	void SetAuthorized(bool bInIsAuthorized)
+	{
+		bIsDeviceAuthorized = bInIsAuthorized;
+	}
+
+	/** Sets the type of the device */
+	void SetDeviceType(const FString InDeviceTypeString)
+	{
+		if (InDeviceTypeString.Contains(TEXT("Browser")))
+		{
+			DeviceType = ETargetDeviceTypes::Browser;
+		}
+		else if (InDeviceTypeString.Contains(TEXT("Console")))
+		{
+			DeviceType = ETargetDeviceTypes::Console;
+		}
+		else if (InDeviceTypeString.Contains(TEXT("Phone")))
+		{
+			DeviceType = ETargetDeviceTypes::Phone;
+		}
+		else if (InDeviceTypeString.Contains(TEXT("Tablet")))
+		{
+			DeviceType = ETargetDeviceTypes::Tablet;
+		}
+		else if (InDeviceTypeString.Contains(TEXT("iPad")))
+		{
+			DeviceType = ETargetDeviceTypes::Tablet;
+		}
+		else
+		{
+			DeviceType = ETargetDeviceTypes::Indeterminate;
+		}
+	}
+
+	/** Sets the connection type (usb/wifi) of the device */
+	void SetDeviceConnectionType(const FString InDeviceConnectionTypeString)
+	{
+		if (InDeviceConnectionTypeString == TEXT("Network"))
+		{
+			DeviceConnectionType = ETargetDeviceConnectionTypes::Wifi;
+		}
+		else if (InDeviceConnectionTypeString == TEXT("USB"))
+		{
+			DeviceConnectionType = ETargetDeviceConnectionTypes::USB;
+		}
+		else if (InDeviceConnectionTypeString == TEXT("Simulator"))
+		{
+			DeviceConnectionType = ETargetDeviceConnectionTypes::Simulator;
+		}
+		else
+		{
+			DeviceConnectionType = ETargetDeviceConnectionTypes::Unknown;
+		}
+	}
+
+	void SetDeviceEndpoint(const FMessageAddress& DeviceAddress)
+	{
+		DeviceEndpoint = DeviceAddress;
+	}
+
+	void SetAppId(const FString& GameName)
+	{
+		AppId = GameName;
+	}
+
+	void SetAppConfiguration(EBuildConfiguration Configuration)
+	{
+		BuildConfiguration = Configuration;
+	}
+
+	void SetIsSimulated(bool IsSimulated)
+	{
+		bIsSimulated = IsSimulated;
+	}
+};
diff --git a/Engine/Source/Developer/IOS/IOSTargetPlatformControls/Private/IOSTargetDeviceOutput.h b/Engine/Source/Developer/IOS/IOSTargetPlatformControls/Private/IOSTargetDeviceOutput.h
new file mode 100644
--- /dev/null
+++ b/Engine/Source/Developer/IOS/IOSTargetPlatformControls/Private/IOSTargetDeviceOutput.h
@@ -0,0 +1,69 @@
+// Copyright Epic Games, Inc. All Rights Reserved.
+
+#pragma once
+
+#include "Containers/UnrealString.h"
+#include "HAL/Runnable.h"
+#include "HAL/RunnableThread.h"
+#include "HAL/PlatformProcess.h"
+#include "Interfaces/ITargetDeviceOutput.h"
+#include "Misc/ConfigCacheIni.h"
+#include "HAL/ThreadSafeCounter.h"
+
+class FIOSTargetDevice;
+
+static FString GetLibImobileDeviceExe(const FString& ExeName)
+{
+	FString ToReturn;
+#if PLATFORM_WINDOWS
+	ToReturn = FPaths::ConvertRelativePathToFull(FPaths::EngineDir() / TEXT("Extras/ThirdPartyNotUE/libimobiledevice/x64/"));
+#elif PLATFORM_MAC
+	ToReturn = FPaths::ConvertRelativePathToFull(FPaths::EngineDir() / TEXT("Extras/ThirdPartyNotUE/libimobiledevice/Mac/"));
+#else
+	UE_LOG(LogIOSDeviceHelper, Error, TEXT("The current platform is unsupported by Libimobile library."));
+#endif
+	ToReturn += ExeName;
+#if PLATFORM_WINDOWS
+	ToReturn += TEXT(".exe");
+#endif
+	return FPaths::FileExists(ToReturn) ? ToReturn : TEXT("");
+}
+
+class FIOSDeviceOutputReaderRunnable : public FRunnable
+{
+public:
+	FIOSDeviceOutputReaderRunnable(const FString& InDeviceUDID, FOutputDevice* Output);
+	
+	// FRunnable interface.
+	virtual bool Init(void) override;
+	virtual void Exit(void) override; 
+	virtual void Stop(void) override;
+	virtual uint32 Run(void) override;
+
+private:
+	bool StartSyslogProcess(void);
+
+	// > 0 if we've been asked to abort work in progress at the next opportunity
+	FThreadSafeCounter	StopTaskCounter;
+	
+	FString				DeviceUDID;
+	FOutputDevice*		Output;
+	void*				SyslogReadPipe;
+	void*				SyslogWritePipe;
+	FProcHandle			SyslogProcHandle;
+};
+
+/**
+ * Implements a IOS target device output.
+ */
+class FIOSTargetDeviceOutput : public ITargetDeviceOutput
+{
+public:
+	bool Init(const FIOSTargetDevice& TargetDevice, FOutputDevice* Output);
+	
+private:
+	TUniquePtr<FRunnableThread>						DeviceOutputThread;
+	FString											DeviceName;
+};
+
+#include "IOSTargetDeviceOutput.inl"
diff --git a/Engine/Source/Developer/IOS/IOSTargetPlatformControls/Private/IOSTargetDeviceOutput.inl b/Engine/Source/Developer/IOS/IOSTargetPlatformControls/Private/IOSTargetDeviceOutput.inl
new file mode 100644
--- /dev/null
+++ b/Engine/Source/Developer/IOS/IOSTargetPlatformControls/Private/IOSTargetDeviceOutput.inl
@@ -0,0 +1,115 @@
+// Copyright Epic Games, Inc. All Rights Reserved.
+
+#include "CoreTypes.h"
+#include "CoreFwd.h"
+#include "Logging/LogMacros.h"
+#include "Common/TcpSocketBuilder.h"
+#include "string.h"
+
+class FIOSDeviceOutputReaderRunnable;
+class FIOSTargetDevice;
+class FIOSTargetDeviceOutput;
+
+inline FIOSDeviceOutputReaderRunnable::FIOSDeviceOutputReaderRunnable(const FString& InDeviceUDID, FOutputDevice* InOutput)
+	: StopTaskCounter(0)
+	, DeviceUDID(InDeviceUDID)
+	, Output(InOutput)
+	, SyslogReadPipe(nullptr)
+	, SyslogWritePipe(nullptr)
+{
+}
+
+inline bool FIOSDeviceOutputReaderRunnable::StartSyslogProcess(void)
+{
+	FString Exe = GetLibImobileDeviceExe("idevicesyslog");
+	FString Params = FString::Printf(TEXT(" -u %s -m [UE]"), *DeviceUDID);
+	SyslogProcHandle = FPlatformProcess::CreateProc(*Exe, *Params, true, false, false, NULL, 0, NULL, SyslogWritePipe);
+	return SyslogProcHandle.IsValid();
+}
+
+inline bool FIOSDeviceOutputReaderRunnable::Init(void)
+{
+	FPlatformProcess::CreatePipe(SyslogReadPipe, SyslogWritePipe);
+	return StartSyslogProcess();
+}
+
+inline void FIOSDeviceOutputReaderRunnable::Exit(void)
+{
+	if (SyslogProcHandle.IsValid())
+	{
+		FPlatformProcess::TerminateProc(SyslogProcHandle);
+	}
+	FPlatformProcess::ClosePipe(SyslogReadPipe, SyslogWritePipe);
+}
+
+inline void FIOSDeviceOutputReaderRunnable::Stop(void)
+{
+	StopTaskCounter.Increment();
+}
+
+inline uint32 FIOSDeviceOutputReaderRunnable::Run(void)
+{
+	FString SyslogOutput;
+
+	while (StopTaskCounter.GetValue() == 0 && SyslogProcHandle.IsValid())
+	{
+		if (!FPlatformProcess::IsProcRunning(SyslogProcHandle))
+		{
+			// When user plugs out USB cable idevicesyslog process stops
+			// Keep trying to restore idevicesyslog connection until code that uses this object will not kill us
+			Output->Serialize(TEXT("Trying to restore connection to device..."), ELogVerbosity::Log, NAME_None);
+			FPlatformProcess::CloseProc(SyslogProcHandle);
+			if (StartSyslogProcess())
+			{
+				FPlatformProcess::Sleep(1.0f);
+			}
+			else
+			{
+				Output->Serialize(TEXT("Failed to start idevicesyslog proccess"), ELogVerbosity::Log, NAME_None);
+				Stop();
+			}
+		}
+		else
+		{
+			SyslogOutput.Append(FPlatformProcess::ReadPipe(SyslogReadPipe));
+
+			if (SyslogOutput.Len() > 0)
+			{
+				TArray<FString> OutputLines;
+				SyslogOutput.ParseIntoArray(OutputLines, TEXT("\n"), false);
+
+				if (!SyslogOutput.EndsWith(TEXT("\n")))
+				{
+					// partial line at the end, do not serialize it until we receive remainder
+					SyslogOutput = OutputLines.Last();
+					OutputLines.RemoveAt(OutputLines.Num() - 1);
+				}
+				else
+				{
+					SyslogOutput.Reset();
+				}
+
+				for (int32 i = 0; i < OutputLines.Num(); ++i)
+				{
+					Output->Serialize(*OutputLines[i], ELogVerbosity::Log, NAME_None);
+				}
+			}
+
+			FPlatformProcess::Sleep(0.1f);
+		}
+	}
+
+	return 0;
+}
+
+inline bool FIOSTargetDeviceOutput::Init(const FIOSTargetDevice& TargetDevice, FOutputDevice* Output)
+{
+	check(Output);
+	// Output will be produced by background thread
+	check(Output->CanBeUsedOnAnyThread());
+	DeviceName = TargetDevice.GetName();
+	auto* Runnable = new FIOSDeviceOutputReaderRunnable(TargetDevice.GetId().GetDeviceName(), Output);
+	DeviceOutputThread = TUniquePtr<FRunnableThread>(FRunnableThread::Create(Runnable, TEXT("FIOSDeviceOutputReaderRunnable")));
+	return true;
+}
+
diff --git a/Engine/Source/Developer/IOS/IOSTargetPlatformControls/Private/IOSTargetPlatformControls.cpp b/Engine/Source/Developer/IOS/IOSTargetPlatformControls/Private/IOSTargetPlatformControls.cpp
new file mode 100644
--- /dev/null
+++ b/Engine/Source/Developer/IOS/IOSTargetPlatformControls/Private/IOSTargetPlatformControls.cpp
@@ -0,0 +1,645 @@
+// Copyright Epic Games, Inc. All Rights Reserved.
+
+/*=============================================================================
+	IOSTargetPlatformControls.cpp: Implements the FIOSTargetPlatformControls class.
+=============================================================================*/
+
+#include "IOSTargetPlatformControls.h"
+#include "Interfaces/IProjectManager.h"
+#include "InstalledPlatformInfo.h"
+#include "HAL/FileManager.h"
+#include "HAL/IConsoleManager.h"
+#include "Misc/Paths.h"
+#include "Misc/App.h"
+#include "Misc/MonitoredProcess.h"
+#include "Logging/MessageLog.h"
+#include "AnalyticsEventAttribute.h"
+
+#if PLATFORM_WINDOWS
+#include "Windows/WindowsHWrapper.h"
+#endif
+#if WITH_ENGINE
+#include "Engine/Texture.h"
+#include "Sound/SoundWave.h"
+#include "TextureResource.h"
+#include "AudioCompressionSettings.h"
+#endif
+
+/* FIOSTargetPlatformControls structors
+ *****************************************************************************/
+
+FIOSTargetPlatformControls::FIOSTargetPlatformControls(bool bInIsTVOS, bool bInIsVisionOS, bool bIsClientOnly, ITargetPlatformSettings* TargetPlatformSettings)
+	// override the ini name up in the base classes, which will go into the FTargetPlatformInfo
+	: TNonDesktopTargetPlatformControlsBase(bIsClientOnly, TargetPlatformSettings, nullptr, bInIsTVOS ? TEXT("TVOS") : bInIsVisionOS ? TEXT("VisionOS") : nullptr)
+	, bIsTVOS(bInIsTVOS)
+	, bIsVisionOS(bInIsVisionOS)
+{
+	// initialize the connected device detector
+	DeviceHelper.OnDeviceConnected().AddRaw(this, &FIOSTargetPlatformControls::HandleDeviceConnected);
+	DeviceHelper.OnDeviceDisconnected().AddRaw(this, &FIOSTargetPlatformControls::HandleDeviceDisconnected);
+	DeviceHelper.Initialize(bIsTVOS || bIsVisionOS);
+}
+
+
+FIOSTargetPlatformControls::~FIOSTargetPlatformControls()
+{
+}
+
+
+/* ITargetPlatform interface
+ *****************************************************************************/
+
+void FIOSTargetPlatformControls::EnableDeviceCheck(bool OnOff)
+{
+	FIOSDeviceHelper::EnableDeviceCheck(OnOff);
+}
+
+void FIOSTargetPlatformControls::GetAllDevices( TArray<ITargetDevicePtr>& OutDevices ) const
+{
+	OutDevices.Reset();
+
+	for (auto Iter = Devices.CreateConstIterator(); Iter; ++Iter)
+	{
+		OutDevices.Add(Iter.Value());
+	}
+}
+
+
+ITargetDevicePtr FIOSTargetPlatformControls::GetDefaultDevice() const
+{
+	if (Devices.Num() > 0)
+	{
+		// first device is the default
+		auto Iter = Devices.CreateConstIterator();
+		if(Iter)
+		{
+			return Iter.Value();
+		}
+	}
+
+	return NULL;
+}
+
+
+ITargetDevicePtr FIOSTargetPlatformControls::GetDevice( const FTargetDeviceId& DeviceId )
+{
+	return Devices.FindRef(DeviceId);
+}
+
+static FString OutputMessage;
+static void OnOutput(FString Message)
+{
+    OutputMessage += Message;
+    UE_LOG(LogTemp, Display, TEXT("%s\n"), *Message);
+}
+
+bool FIOSTargetPlatformControls::IsSdkInstalled(bool bProjectHasCode, FString& OutTutorialPath) const
+{
+#if PLATFORM_MAC
+	OutTutorialPath = FString("Shared/Tutorials/InstallingXCodeTutorial");
+
+	// run xcode-select and get the location of Xcode
+	FString CmdExe = TEXT("/usr/bin/xcode-select");
+	FString CommandLine = FString::Printf(TEXT("--print-path"));
+	TSharedPtr<FMonitoredProcess> IPPProcess = MakeShareable(new FMonitoredProcess(CmdExe, CommandLine, true));
+	OutputMessage = TEXT("");
+	IPPProcess->OnOutput().BindStatic(&OnOutput);
+	IPPProcess->Launch();
+	while (IPPProcess->Update())
+	{
+		FPlatformProcess::Sleep(0.01f);
+	}
+	int RetCode = IPPProcess->GetReturnCode();
+//	UE_LOG(LogTemp, Display, TEXT("%s"), *OutputMessage);
+
+	bool biOSSDKInstalled = IFileManager::Get().DirectoryExists(*OutputMessage);
+#else
+	OutTutorialPath = FString("/Engine/Tutorial/Mobile/InstallingiTunesTutorial.InstallingiTunesTutorial");
+
+	// On windows we check if itunes is installed - Perhaps someday make this its own check instead of piggy packing on the SDK check which will create a unintuitive error message when it fails
+
+	// The logic here is to assume the correct Apple dll does not exist and then check the various locations it could be in, setting this to true when it is found
+	// Code is structured for clarity not performance
+	bool biOSSDKInstalled = false;
+
+	HKEY hKey;
+	TCHAR dllPath[256];
+	unsigned long pathSize = 256;
+	
+	// Add future version checks here
+
+	// Check for iTunes 12
+	if (!biOSSDKInstalled && RegOpenKeyEx(HKEY_LOCAL_MACHINE, TEXT("SOFTWARE\\Wow6432Node\\Apple Inc.\\Apple Mobile Device Support\\Shared"), 0, KEY_READ, &hKey) == ERROR_SUCCESS)
+	{
+		if (RegQueryValueEx(hKey, TEXT("MobileDeviceDLL"), 0, NULL, (BYTE*)dllPath, &pathSize) == ERROR_SUCCESS
+			&& IFileManager::Get().FileSize(*FString(dllPath)) != INDEX_NONE)
+		{
+			biOSSDKInstalled = true;
+		}
+		else
+		{
+			// iTunes >= 12.7 doesn't have a key specifying the 32-bit DLL but it does have a ASMapiInterfaceDLL key and MobileDevice.dll is in usually in the same directory
+			if (RegQueryValueEx(hKey, TEXT("ASMapiInterfaceDLL"), 0, NULL, (BYTE*)dllPath, &pathSize) == ERROR_SUCCESS)
+			{
+				FString MobileDeviceDLLPath = FString(dllPath);
+				int32 Index = MobileDeviceDLLPath.Find(TEXT("\\"), ESearchCase::CaseSensitive, ESearchDir::FromEnd);
+				if(Index != INDEX_NONE)
+				{
+					MobileDeviceDLLPath = MobileDeviceDLLPath.Left(Index+1) + TEXT("MobileDevice.dll");
+					if (IFileManager::Get().FileSize(*MobileDeviceDLLPath) != INDEX_NONE)
+					{
+						biOSSDKInstalled = true;
+					}
+				}
+			}
+		}
+	}
+
+	// Check for iTunes 12, Windows Store version
+	if (!biOSSDKInstalled)
+	{
+		HKEY hPackagesKey;
+		FString PackagesKeyName(TEXT("Software\\Classes\\Local Settings\\Software\\Microsoft\\Windows\\CurrentVersion\\AppModel\\PackageRepository\\Packages"));
+		if (RegOpenKeyEx(HKEY_LOCAL_MACHINE, *PackagesKeyName, 0, KEY_READ, &hPackagesKey) == ERROR_SUCCESS)
+		{
+			DWORD NumSubKeys = 0;
+			DWORD LongestSubKeyLength;
+			if (RegQueryInfoKey(hPackagesKey, NULL, NULL, NULL, &NumSubKeys, &LongestSubKeyLength, NULL, NULL, NULL, NULL, NULL, NULL) == ERROR_SUCCESS)
+			{
+				static const uint32 MAX_SUBKEY_LENGTH = 512;
+				check(LongestSubKeyLength < MAX_SUBKEY_LENGTH);
+
+				for (uint32 i = 0; i < NumSubKeys; i++)
+				{
+					TCHAR SubKeyName[MAX_SUBKEY_LENGTH];
+					DWORD SubKeyLength = MAX_SUBKEY_LENGTH;
+					if (RegEnumKeyEx(hPackagesKey, i, SubKeyName, &SubKeyLength, NULL, NULL, NULL, NULL) == ERROR_SUCCESS)
+					{
+						check(SubKeyLength < MAX_SUBKEY_LENGTH);
+
+						if (FString(SubKeyName).Contains(TEXT("AppleInc.iTunes")) && (FString(SubKeyName).Contains(TEXT("_x64")) || FString(SubKeyName).Contains(TEXT("_x86"))))
+						{
+							HKEY iTunesKey;
+							TCHAR FullPackageSubKeyName[MAX_SUBKEY_LENGTH + MAX_SUBKEY_LENGTH + 1];
+							_stprintf_s(FullPackageSubKeyName, MAX_SUBKEY_LENGTH + MAX_SUBKEY_LENGTH + 1, TEXT("%s\\%s"), *PackagesKeyName, SubKeyName);
+							if (RegOpenKeyEx(HKEY_LOCAL_MACHINE, FullPackageSubKeyName, 0, KEY_READ, &iTunesKey) == ERROR_SUCCESS)
+							{
+								if (RegQueryValueEx(iTunesKey, TEXT("Path"), NULL, NULL, (BYTE*)dllPath, &pathSize) == ERROR_SUCCESS)
+								{
+									TCHAR dllFullPath[MAX_SUBKEY_LENGTH + MAX_SUBKEY_LENGTH + 32];
+									_stprintf_s(dllFullPath, MAX_SUBKEY_LENGTH + MAX_SUBKEY_LENGTH + 32, TEXT("%s\\AMDS32\\MobileDevice.dll"), dllPath);
+
+									if (IFileManager::Get().FileSize(*FString(dllFullPath)) != INDEX_NONE)
+									{
+										biOSSDKInstalled = true;
+										break;
+									}
+								}
+							}
+						}
+					}
+				}
+			}
+		}
+	}
+
+	// Check for iTunes 11
+	if(!biOSSDKInstalled
+		&& RegOpenKeyEx(HKEY_LOCAL_MACHINE, TEXT("SOFTWARE\\Wow6432Node\\Apple Inc.\\Apple Mobile Device Support\\Shared"), 0, KEY_READ, &hKey) == ERROR_SUCCESS
+		&& RegQueryValueEx(hKey, TEXT("iTunesMobileDeviceDLL"), 0, NULL, (BYTE*)dllPath, &pathSize) == ERROR_SUCCESS
+		&&  IFileManager::Get().FileSize(*FString(dllPath)) != INDEX_NONE)
+			{
+		biOSSDKInstalled = true;
+	}
+
+#endif
+	return biOSSDKInstalled;
+}
+
+int32 FIOSTargetPlatformControls::CheckRequirements(bool bProjectHasCode, EBuildConfiguration Configuration, bool bRequiresAssetNativization, FString& OutTutorialPath, FString& OutDocumentationPath, FText& CustomizedLogMessage) const
+{
+	OutDocumentationPath = TEXT("Platforms/iOS/QuickStart/6");
+
+	int32 bReadyToBuild = ETargetPlatformReadyStatus::Ready; // @todo How do we check that the iOS SDK is installed when building from Windows? Is that even possible?
+	if (!IsSdkInstalled(bProjectHasCode, OutTutorialPath))
+	{
+		bReadyToBuild |= ETargetPlatformReadyStatus::SDKNotFound;
+	}
+#if PLATFORM_MAC
+	OutTutorialPath = FString("/Engine/Tutorial/Installation/InstallingXCodeTutorial.InstallingXCodeTutorial");
+    // shell to certtool
+#else
+	if (!FInstalledPlatformInfo::Get().IsValidPlatform(GetPlatformInfo().UBTPlatformString, EProjectType::Code))
+	{
+		if (bProjectHasCode)
+		{
+			OutTutorialPath = FString("/Engine/Tutorial/Mobile/iOSonPCRestrictions.iOSonPCRestrictions");
+			bReadyToBuild |= ETargetPlatformReadyStatus::CodeUnsupported;
+		}
+
+		FText Reason;
+		if (RequiresTempTarget(bProjectHasCode, Configuration, bRequiresAssetNativization, Reason))
+		{
+			OutTutorialPath = FString("/Engine/Tutorial/Mobile/iOSonPCValidPlugins.iOSonPCValidPlugins");
+			bReadyToBuild |= ETargetPlatformReadyStatus::PluginsUnsupported;
+		}
+	}
+#endif
+
+	// shell to IPP and get the status of the provision and cert
+
+	bool bForDistribtion = false;
+	GConfig->GetBool(TEXT("/Script/UnrealEd.ProjectPackagingSettings"), TEXT("ForDistribution"), bForDistribtion, GGameIni);
+
+	FString BundleIdentifier;
+	GConfig->GetString(TEXT("/Script/IOSRuntimeSettings.IOSRuntimeSettings"), TEXT("BundleIdentifier"), BundleIdentifier, GEngineIni);
+	BundleIdentifier = BundleIdentifier.Replace(TEXT("[PROJECT_NAME]"), FApp::GetProjectName());
+	BundleIdentifier = BundleIdentifier.Replace(TEXT("_"), TEXT(""));
+
+	bool bAutomaticSigning = false;
+	GConfig->GetBool(TEXT("/Script/IOSRuntimeSettings.IOSRuntimeSettings"), TEXT("bAutomaticSigning"), bAutomaticSigning, GEngineIni);
+
+	FString TeamID;
+	GConfig->GetString(TEXT("/Script/IOSRuntimeSettings.IOSRuntimeSettings"), TEXT("IOSTeamID"), TeamID, GEngineIni);
+
+	FString ProjectPath = FPaths::ConvertRelativePathToFull(FPaths::GetProjectFilePath());
+#if PLATFORM_MAC
+    FString CmdExe = TEXT("/bin/sh");
+    FString ScriptPath = FPaths::ConvertRelativePathToFull(FPaths::EngineDir() / TEXT("Build/BatchFiles/RunDotnet.sh"));
+    FString IPPPath = FPaths::ConvertRelativePathToFull(FPaths::EngineDir() / TEXT("Binaries/DotNET/IOS/IPhonePackager.exe"));
+	FString CommandLine = FString::Printf(TEXT("\"%s\" \"%s\" Validate Engine -project \"%s\" -bundlename \"%s\" -teamID \"%s\" %s %s"), *ScriptPath, *IPPPath, *ProjectPath, *(BundleIdentifier), *(TeamID), (bForDistribtion ? TEXT("-distribution") : TEXT("")), bAutomaticSigning ? TEXT("-autosigning") : TEXT(""));
+#else
+	FString CmdExe = FPaths::ConvertRelativePathToFull(FPaths::EngineDir() / TEXT("Binaries/DotNET/IOS/IPhonePackager.exe"));
+	FString CommandLine = FString::Printf(TEXT("Validate Engine -project \"%s\" -bundlename \"%s\" %s"), *ProjectPath, *(BundleIdentifier), (bForDistribtion ? TEXT("-distribution") : TEXT("")) );
+	FString RemoteServerName;
+	FString SecondaryRemoteServerName;
+	FString RSyncUsername;
+	FString SecondaryRSyncUsername;
+	GConfig->GetString(TEXT("/Script/IOSRuntimeSettings.IOSRuntimeSettings"), TEXT("RemoteServerName"), RemoteServerName, GEngineIni);
+	GConfig->GetString(TEXT("/Script/IOSRuntimeSettings.IOSRuntimeSettings"), TEXT("RSyncUsername"), RSyncUsername, GEngineIni);
+	GConfig->GetString(TEXT("/Script/IOSRuntimeSettings.IOSRuntimeSettings"), TEXT("SecondaryRemoteServerName"), SecondaryRemoteServerName, GEngineIni);
+	GConfig->GetString(TEXT("/Script/IOSRuntimeSettings.IOSRuntimeSettings"), TEXT("SecondaryRSyncUsername"), SecondaryRSyncUsername, GEngineIni);
+	if (RemoteServerName.Len() == 0 || RSyncUsername.Len() == 0)
+	{
+		bReadyToBuild |= ETargetPlatformReadyStatus::RemoveServerNameEmpty;
+	}
+
+#endif
+	if (bIsTVOS)
+	{
+		CommandLine += " -tvos";
+	}
+	TSharedPtr<FMonitoredProcess> IPPProcess = MakeShareable(new FMonitoredProcess(CmdExe, CommandLine, true));
+	OutputMessage = TEXT("");
+	IPPProcess->OnOutput().BindStatic(&OnOutput);
+	IPPProcess->Launch();
+	while(IPPProcess->Update())
+	{
+		FPlatformProcess::Sleep(0.01f);
+	}
+	int RetCode = IPPProcess->GetReturnCode();
+//	UE_LOG(LogTemp, Display, TEXT("%s"), *OutputMessage);
+	if (RetCode == 14)
+	{
+		OutTutorialPath = FString("/Engine/Tutorial/Mobile/CreatingInfoPlist.CreatingInfoPlist");
+		bReadyToBuild |= ETargetPlatformReadyStatus::ManifestNotFound;
+	}
+	else if (RetCode == 13)
+	{
+		OutTutorialPath = FString("/Engine/Tutorial/Mobile/CreatingSigningCertAndProvisionTutorial.CreatingSigningCertAndProvisionTutorial");
+		bReadyToBuild |= ETargetPlatformReadyStatus::SigningKeyNotFound;
+		bReadyToBuild |= ETargetPlatformReadyStatus::ProvisionNotFound;
+	}
+	else if (RetCode == 12)
+	{
+		OutTutorialPath = FString("/Engine/Tutorial/Mobile/CreatingSigningCertAndProvisionTutorial.CreatingSigningCertAndProvisionTutorial");
+		bReadyToBuild |= ETargetPlatformReadyStatus::SigningKeyNotFound;
+	}
+	else if (RetCode == 11)
+	{
+		OutTutorialPath = FString("/Engine/Tutorial/Mobile/CreatingSigningCertAndProvisionTutorial.CreatingSigningCertAndProvisionTutorial");
+		bReadyToBuild |= ETargetPlatformReadyStatus::ProvisionNotFound;
+	}
+
+	{
+		TArray<FString> FoundIconFiles;
+		FString Wildcard = FPaths::Combine(FPaths::ProjectDir(), TEXT("Build"), TEXT("IOS"), TEXT("Resources"), TEXT("Graphics"), TEXT("Icon*.png"));
+		IFileManager::Get().FindFiles(FoundIconFiles, *Wildcard, true, false);
+		if (FoundIconFiles.Num() > 0)
+		{
+			bReadyToBuild |= ETargetPlatformReadyStatus::CodeBuildRequired;
+		}
+	}
+
+	return bReadyToBuild;
+}
+
+
+
+/* FIOSTargetPlatformControls callbacks
+ *****************************************************************************/
+
+void FIOSTargetPlatformControls::HandlePongMessage( const FIOSLaunchDaemonPong& Message, const TSharedRef<IMessageContext, ESPMode::ThreadSafe>& Context )
+{
+	FTargetDeviceId DeviceId;
+	FTargetDeviceId::Parse(Message.DeviceID, DeviceId);
+
+	FIOSTargetDevicePtr& Device = Devices.FindOrAdd(DeviceId);
+
+	if (!Device.IsValid())
+	{
+		Device = MakeShareable(new FIOSTargetDevice(*this));
+
+		Device->SetFeature(ETargetDeviceFeatures::Reboot, Message.bCanReboot);
+		Device->SetFeature(ETargetDeviceFeatures::PowerOn, Message.bCanPowerOn);
+		Device->SetFeature(ETargetDeviceFeatures::PowerOff, Message.bCanPowerOff);
+		Device->SetDeviceId(DeviceId);
+		Device->SetDeviceName(Message.DeviceName);
+		Device->SetDeviceType(Message.DeviceType);
+		Device->SetModelId(Message.DeviceModelId);
+		Device->SetOSVersion(Message.DeviceOSVersion);
+		Device->SetDeviceConnectionType(Message.DeviceConnectionType);
+		Device->SetDeviceEndpoint(Context->GetSender());
+		Device->SetIsSimulated(Message.DeviceID.Contains(TEXT("Simulator")));
+
+		ITargetPlatformControls::OnDeviceDiscovered().Broadcast(Device.ToSharedRef());
+	}
+
+	Device->LastPinged = FDateTime::UtcNow();
+}
+
+void FIOSTargetPlatformControls::HandleDeviceConnected(const FIOSLaunchDaemonPong& Message)
+{
+	FTargetDeviceId DeviceId;
+	FTargetDeviceId::Parse(Message.DeviceID, DeviceId);
+	
+	FIOSTargetDevicePtr& Device = Devices.FindOrAdd(DeviceId);
+	
+	if (!Device.IsValid())
+	{
+		bool bIsTVOSDevice = Message.DeviceType.Contains(TEXT("AppleTV"));
+		bool bIsVisionOSDevice = Message.DeviceType.Contains(TEXT("RealityDevice"));
+		bool bIsIOSDevice = !bIsTVOSDevice && !bIsVisionOSDevice;
+		
+		bool bIsIOS = !bIsTVOS && !bIsVisionOS;
+
+		if ((bIsIOS && bIsIOSDevice) || (bIsTVOS && bIsTVOSDevice) || (bIsVisionOS && bIsVisionOSDevice))
+		{
+			Device = MakeShareable(new FIOSTargetDevice(*this));
+
+			Device->SetFeature(ETargetDeviceFeatures::Reboot, Message.bCanReboot);
+			Device->SetFeature(ETargetDeviceFeatures::PowerOn, Message.bCanPowerOn);
+			Device->SetFeature(ETargetDeviceFeatures::PowerOff, Message.bCanPowerOff);
+			Device->SetDeviceId(DeviceId);
+			Device->SetDeviceName(Message.DeviceName);
+			Device->SetAuthorized(Message.bIsAuthorized);
+			Device->SetDeviceType(Message.DeviceType);
+			Device->SetModelId(Message.DeviceModelId);
+			Device->SetOSVersion(Message.DeviceOSVersion);
+			Device->SetDeviceConnectionType(Message.DeviceConnectionType);
+			Device->SetIsSimulated(Message.DeviceID.Contains(TEXT("Simulator")));
+
+			OnDeviceDiscovered().Broadcast(Device.ToSharedRef());
+		}
+		else
+		{
+			return;
+		}
+	}
+	
+	// Add a very long time period to prevent the devices from getting disconnected due to a lack of pong messages
+	Device->LastPinged = FDateTime::UtcNow() + FTimespan::FromDays(100.0);
+}
+
+
+void FIOSTargetPlatformControls::HandleDeviceDisconnected(const FIOSLaunchDaemonPong& Message)
+{
+	FTargetDeviceId DeviceId;
+	FTargetDeviceId::Parse(Message.DeviceID, DeviceId);
+	
+	FIOSTargetDevicePtr& Device = Devices.FindOrAdd(DeviceId);
+	
+	if (Device.IsValid())
+	{
+		ITargetPlatformControls::OnDeviceLost().Broadcast(Device.ToSharedRef());
+		Devices.Remove(DeviceId);
+	}
+}
+
+
+/* ITargetPlatform interface
+ *****************************************************************************/
+static bool SupportsMetal()
+{
+	// default to NOT supporting metal
+	bool bSupportsMetal = false;
+	GConfig->GetBool(TEXT("/Script/IOSRuntimeSettings.IOSRuntimeSettings"), TEXT("bSupportsMetal"), bSupportsMetal, GEngineIni);
+	return bSupportsMetal;
+}
+
+static bool SupportsMetalMRT()
+{
+	// default to NOT supporting metal MRT
+	bool bSupportsMetalMRT = false;
+	GConfig->GetBool(TEXT("/Script/IOSRuntimeSettings.IOSRuntimeSettings"), TEXT("bSupportsMetalMRT"), bSupportsMetalMRT, GEngineIni);
+	return bSupportsMetalMRT;
+}
+
+static bool SupportsA8Devices()
+{
+    // default to NOT supporting A8 devices
+    bool bSupportAppleA8 = false;
+    GConfig->GetBool(TEXT("/Script/IOSRuntimeSettings.IOSRuntimeSettings"), TEXT("bSupportAppleA8"), bSupportAppleA8, GEngineIni);
+    return bSupportAppleA8;
+}
+
+bool FIOSTargetPlatformControls::CanSupportRemoteShaderCompile() const
+{
+	// for 4.22 we are disabling support for XGE Shader compile on IOS
+	bool bRemoteCompilingEnabled = false;
+	GConfig->GetBool(TEXT("/Script/IOSRuntimeSettings.IOSRuntimeSettings"), TEXT("EnableRemoteShaderCompile"), bRemoteCompilingEnabled, GEngineIni);
+	return false; // !bRemoteCompilingEnabled;
+}
+
+void FIOSTargetPlatformControls::GetPlatformSpecificProjectAnalytics( TArray<FAnalyticsEventAttribute>& AnalyticsParamArray ) const
+{
+	TNonDesktopTargetPlatformControlsBase<FIOSPlatformProperties>::GetPlatformSpecificProjectAnalytics( AnalyticsParamArray );
+
+	AppendAnalyticsEventAttributeArray(AnalyticsParamArray,
+		TEXT("SupportsMetalMRT"), SupportsMetalMRT()
+	);
+}
+
+#if WITH_ENGINE
+
+static const FName NameASTC_RGB_HDR(TEXT("ASTC_RGB_HDR"));
+static const FName NameBC5(TEXT("BC5"));
+static const FName NameASTC_NormalLA(TEXT("ASTC_NormalLA"));
+
+// we remap some of the defaults
+static const FName FormatRemap[] =
+{
+	// original				ASTC
+	FName(TEXT("AutoDXT")),	FName(TEXT("ASTC_RGBAuto")),
+	FName(TEXT("DXT1")),	FName(TEXT("ASTC_RGB")),
+	FName(TEXT("DXT5")),	FName(TEXT("ASTC_RGBA")),
+	FName(TEXT("DXT5n")),	FName(TEXT("ASTC_NormalAG")),
+	NameBC5,				FName(TEXT("ASTC_NormalRG")),
+	FName(TEXT("BC4")),		FName(TEXT("ETC2_R11")),
+	FName(TEXT("BC6H")),	NameASTC_RGB_HDR,
+	FName(TEXT("BC7")),		FName(TEXT("ASTC_RGBA_HQ"))
+};
+static const FName NameG8(TEXT("G8"));
+static const FName NameRGBA16F(TEXT("RGBA16F"));
+
+void FIOSTargetPlatformControls::GetTextureFormats( const UTexture* Texture, TArray< TArray<FName> >& OutFormats) const
+{
+	check(Texture);
+
+	const int32 NumLayers = Texture->Source.GetNumLayers();
+	
+	TArray<FName>& TextureFormatNames = OutFormats.AddDefaulted_GetRef();
+	TextureFormatNames.Reserve(NumLayers);
+
+	// optionaly compress landscape weightmaps for a mobile rendering
+	// @todo Oodle: this should not be here
+	//	should be in GetDefaultTextureFormatNamePerLayer
+	//	so that 4x4 checks can be applied correctly, etc.
+	static const auto CompressLandscapeWeightMapsVar = IConsoleManager::Get().FindTConsoleVariableDataInt(TEXT("r.Mobile.CompressLandscapeWeightMaps"));
+	static const bool bCompressLandscapeWeightMaps = (CompressLandscapeWeightMapsVar && CompressLandscapeWeightMapsVar->GetValueOnAnyThread() != 0);
+	if (Texture->LODGroup == TEXTUREGROUP_Terrain_Weightmap && bCompressLandscapeWeightMaps)
+	{
+		TextureFormatNames.Init(FName(TEXT("AutoDXT")), NumLayers);
+	}
+
+	// if we didn't assign anything specially, then use the defaults
+	if (TextureFormatNames.Num() == 0)
+	{
+        int32 BlockSize = 1;
+		// Compressed volume textures require MTLGPUFamilyApple3 or later
+		// min spec for TVOS is AppleTV HD which is MTLGPUFamilyApple2 (A8)
+		bool bSupportCompressedVolumeTexture = !bIsTVOS && !SupportsA8Devices();
+		bool bSupportFilteredFloat32Textures = false;
+		GetDefaultTextureFormatNamePerLayer(TextureFormatNames, this->GetTargetPlatformSettings(), this, Texture, bSupportCompressedVolumeTexture, BlockSize, bSupportFilteredFloat32Textures);
+	}
+
+	// L+A mode for normal map compression
+	const bool bSupportsNormalLA = GetTargetPlatformSettings()->SupportsFeature(ETargetPlatformFeatures::NormalmapLAEncodingMode);
+
+	// include the formats we want
+	for (FName& TextureFormatName : TextureFormatNames)
+	{
+		if (bSupportsNormalLA && TextureFormatName == NameBC5)
+		{
+			TextureFormatName = NameASTC_NormalLA;
+			continue;
+		}
+		
+		for (int32 RemapIndex = 0; RemapIndex < UE_ARRAY_COUNT(FormatRemap); RemapIndex += 2)
+		{
+			if (TextureFormatName == FormatRemap[RemapIndex])
+			{
+				TextureFormatName = FormatRemap[RemapIndex + 1];
+				break;
+			}
+		}
+	}
+	
+	bool bSupportASTCHDR = GetTargetPlatformSettings()->UsesASTCHDR();
+
+	if ( ! bSupportASTCHDR )
+	{
+		for (FName& TextureFormatName : TextureFormatNames)
+		{
+			if ( TextureFormatName == NameASTC_RGB_HDR )
+			{
+				TextureFormatName = NameRGBA16F;
+			}
+		}
+	}
+
+	bool bEnableSimulatorSupport = false;
+	GConfig->GetBool(TEXT("/Script/IOSRuntimeSettings.IOSRuntimeSettings"), TEXT("bEnableSimulatorSupport"), bEnableSimulatorSupport, GEngineIni);
+
+	for (FName& TextureFormatName : OutFormats.Last())
+	{
+		if (Texture->GetTextureClass() == ETextureClass::Cube) 
+		{
+			FTextureFormatSettings FormatSettings;
+			Texture->GetDefaultFormatSettings(FormatSettings);
+			// TC_EncodedReflectionCapture is no longer used and could be deleted
+			if (FormatSettings.CompressionSettings == TC_EncodedReflectionCapture && !FormatSettings.CompressionNone)
+			{
+				TextureFormatName = FName(TEXT("ETC2_RGBA"));
+			}
+		}
+
+		// Currently (Xcode14), the iOS Simulator does not support compressed Volume textures.
+		if (bEnableSimulatorSupport && Texture->GetTextureClass() == ETextureClass::Volume)
+		{
+			FTextureFormatSettings FormatSettings;
+			Texture->GetDefaultFormatSettings(FormatSettings);
+			TextureFormatName = FName(TEXT("RGB8"));
+		}
+	}
+}
+
+void FIOSTargetPlatformControls::GetAllTextureFormats(TArray<FName>& OutFormats) const 
+{
+	bool bFoundRemap = false;
+
+	GetAllDefaultTextureFormats(this->GetTargetPlatformSettings(), OutFormats);
+
+	for (int32 RemapIndex = 0; RemapIndex < UE_ARRAY_COUNT(FormatRemap); RemapIndex += 2)
+	{
+		OutFormats.Remove(FormatRemap[RemapIndex+0]);
+	}
+
+	// include the formats we want
+	for (int32 RemapIndex = 0; RemapIndex < UE_ARRAY_COUNT(FormatRemap); RemapIndex += 2)
+	{
+		OutFormats.AddUnique(FormatRemap[RemapIndex + 1]);
+	}
+}
+
+FName FIOSTargetPlatformControls::FinalizeVirtualTextureLayerFormat(FName Format) const
+{
+#if WITH_EDITOR
+
+	// VirtualTexture Format was already run through the ordinary texture remaps to change AutoDXT to ASTC or ETC
+	// this then runs again
+	// currently it forces all ASTC to ETC
+	// this is needed because the runtime virtual texture encoder only supports ETC
+	
+	// code dupe with AndroidTargetPlatform
+
+	const static FName NameETC2_RGB(TEXT("ETC2_RGB"));
+	const static FName NameETC2_RGBA(TEXT("ETC2_RGBA"));
+	const static FName NameAutoETC2(TEXT("AutoETC2"));
+	const static FName NameETC2_RG11(TEXT("ETC2_RG11"));
+
+	// Remap non-ETC variants to ETC
+	const static FName ETCRemap[][2] =
+	{
+		{ { FName(TEXT("ASTC_RGB")) },			{ NameETC2_RGB } },
+		{ { FName(TEXT("ASTC_RGBA")) },			{ NameETC2_RGBA } },
+		{ { FName(TEXT("ASTC_RGBAuto")) },		{ NameAutoETC2 } },
+		{ { FName(TEXT("ASTC_RGBA_HQ")) },		{ NameETC2_RGBA } },
+//		{ { FName(TEXT("ASTC_RGB_HDR")) },		{ NameRGBA16F } }, // ?
+		{ { FName(TEXT("ASTC_NormalAG")) },		{ NameETC2_RGB } },
+		{ { FName(TEXT("ASTC_NormalRG")) },		{ NameETC2_RG11 } },
+		{ { FName(TEXT("ASTC_NormalLA")) },	    { NameETC2_RG11 } },
+	};
+
+	for (int32 RemapIndex = 0; RemapIndex < UE_ARRAY_COUNT(ETCRemap); RemapIndex++)
+	{
+		if (ETCRemap[RemapIndex][0] == Format)
+		{
+			return ETCRemap[RemapIndex][1];
+		}
+	}
+#endif
+	return Format;
+}
+#endif // WITH_ENGINE
+
diff --git a/Engine/Source/Developer/IOS/IOSTargetPlatformControls/Private/IOSTargetPlatformControls.h b/Engine/Source/Developer/IOS/IOSTargetPlatformControls/Private/IOSTargetPlatformControls.h
new file mode 100644
--- /dev/null
+++ b/Engine/Source/Developer/IOS/IOSTargetPlatformControls/Private/IOSTargetPlatformControls.h
@@ -0,0 +1,114 @@
+// Copyright Epic Games, Inc. All Rights Reserved.
+
+/*=============================================================================
+	IOSTargetPlatformControls.h: Declares the FIOSTargetPlatformControls class.
+=============================================================================*/
+
+#pragma once
+
+#include "CoreMinimal.h"
+#include "Common/TargetPlatformControlsBase.h"
+#include "IOS/IOSPlatformProperties.h"
+#include "Containers/Ticker.h"
+#include "IOSMessageProtocol.h"
+#include "IMessageContext.h"
+#include "IOSTargetDevice.h"
+#include "IOSDeviceHelper.h"
+#include "Misc/ConfigCacheIni.h"
+
+
+#if WITH_ENGINE
+#include "AudioCompressionSettings.h"
+#endif // WITH_ENGINE
+
+/**
+ * FIOSTargetPlatformControls, abstraction for cooking iOS platforms
+ */
+class FIOSTargetPlatformControls : public TNonDesktopTargetPlatformControlsBase<FIOSPlatformProperties>
+{
+public:
+
+	/**
+	 * Default constructor.
+	 */
+	IOSTARGETPLATFORMCONTROLS_API FIOSTargetPlatformControls(bool bInISTVOS, bool bInIsVisionOS, bool bInIsClientOnly, ITargetPlatformSettings* TargetPlatformSettings);
+
+	/**
+	 * Destructor.
+	 */
+	~FIOSTargetPlatformControls();
+
+public:
+
+	virtual void EnableDeviceCheck(bool OnOff) override;
+
+	virtual void GetAllDevices( TArray<ITargetDevicePtr>& OutDevices ) const override;
+
+	virtual bool GenerateStreamingInstallManifest(const TMultiMap<FString, int32>& PakchunkMap, const TSet<int32>& PakchunkIndicesInUse) const override
+	{
+		return true;
+	}
+
+	virtual ITargetDevicePtr GetDefaultDevice( ) const override;
+
+	virtual ITargetDevicePtr GetDevice( const FTargetDeviceId& DeviceId ) override;
+
+	virtual bool CanSupportRemoteShaderCompile() const override;
+
+	virtual bool IsSdkInstalled(bool bProjectHasCode, FString& OutTutorialPath) const override;
+	virtual int32 CheckRequirements(bool bProjectHasCode, EBuildConfiguration Configuration, bool bRequiresAssetNativization, FString& OutTutorialPath, FString& OutDocumentationPath, FText& CustomizedLogMessage) const override;
+
+	virtual void GetPlatformSpecificProjectAnalytics( TArray<struct FAnalyticsEventAttribute>& AnalyticsParamArray ) const override;
+
+#if WITH_ENGINE
+
+	virtual void GetTextureFormats( const UTexture* Texture, TArray< TArray<FName> >& OutFormats) const override;
+
+	virtual void GetAllTextureFormats( TArray<FName>& OutFormats) const override;
+
+	virtual FName FinalizeVirtualTextureLayerFormat(FName Format) const override;
+
+#endif // WITH_ENGINE
+
+	virtual void GetBuildProjectSettingKeys(FString& OutSection, TArray<FString>& InBoolKeys, TArray<FString>& InIntKeys, TArray<FString>& InStringKeys) const override
+	{
+		OutSection = TEXT("/Script/IOSRuntimeSettings.IOSRuntimeSettings");
+		InBoolKeys.Add(TEXT("bGeneratedSYMFile"));
+		InBoolKeys.Add(TEXT("bGeneratedSYMBundle"));
+		InBoolKeys.Add(TEXT("bGenerateXCArchive"));
+		if (bIsTVOS)
+		{
+			InStringKeys.Add(TEXT("MinimumTVOSVersion"));
+		}
+		else
+		{
+			InStringKeys.Add(TEXT("MinimumiOSVersion"));
+		}
+	}
+
+	//~ Begin ITargetPlatform Interface
+
+private:
+
+	// Handles received pong messages from the LauncherDaemon.
+	void HandlePongMessage( const FIOSLaunchDaemonPong& Message, const TSharedRef<IMessageContext, ESPMode::ThreadSafe>& Context );
+
+    void HandleDeviceConnected( const FIOSLaunchDaemonPong& Message );
+    void HandleDeviceDisconnected( const FIOSLaunchDaemonPong& Message );
+
+private:
+	
+	// true if this is targeting TVOS vs IOS
+	bool bIsTVOS;
+	bool bIsVisionOS;
+
+	// Contains all discovered IOSTargetDevices over the network.
+	TMap<FTargetDeviceId, FIOSTargetDevicePtr> Devices;
+
+	// Holds the message endpoint used for communicating with the LaunchDaemon.
+	TSharedPtr<FMessageEndpoint, ESPMode::ThreadSafe> MessageEndpoint;
+
+    // holds usb device helper
+	FIOSDeviceHelper DeviceHelper;
+
+};
diff --git a/Engine/Source/Developer/IOS/IOSTargetPlatformControls/Private/IOSTargetPlatformControlsModule.cpp b/Engine/Source/Developer/IOS/IOSTargetPlatformControls/Private/IOSTargetPlatformControlsModule.cpp
new file mode 100644
--- /dev/null
+++ b/Engine/Source/Developer/IOS/IOSTargetPlatformControls/Private/IOSTargetPlatformControlsModule.cpp
@@ -0,0 +1,34 @@
+// Copyright Epic Games, Inc. All Rights Reserved.
+
+#include "IOSTargetPlatformControls.h"
+#include "Interfaces/ITargetPlatformControlsModule.h"
+#include "Interfaces/ITargetPlatformSettingsModule.h"
+#include "Common/TargetPlatformSettingsBase.h"
+#include "Modules/ModuleManager.h"
+
+
+/**
+ * Module for iOS as a target platform controls
+ */
+class FIOSTargetPlatformControlsModule : public ITargetPlatformControlsModule
+{
+public:
+	virtual void GetTargetPlatformControls(TArray<ITargetPlatformControls*>& TargetPlatforms, FName& PlatformSettingsModuleName) override
+	{
+		if (FIOSTargetPlatformControls::IsUsable())
+		{
+			ITargetPlatformSettingsModule* ModuleSettings = FModuleManager::GetModulePtr<ITargetPlatformSettingsModule>(PlatformSettingsModuleName);
+			if (ModuleSettings != nullptr)
+			{
+				TArray<ITargetPlatformSettings*> TargetPlatformSettings = ModuleSettings->GetTargetPlatformSettings();
+				check(TargetPlatformSettings.Num() == 1);
+
+				TargetPlatforms.Add(new FIOSTargetPlatformControls(false, false, false, TargetPlatformSettings[0]));
+				TargetPlatforms.Add(new FIOSTargetPlatformControls(false, false, true, TargetPlatformSettings[0]));
+			}
+		}
+	}
+};
+
+
+IMPLEMENT_MODULE(FIOSTargetPlatformControlsModule, IOSTargetPlatformControls);
diff --git a/Engine/Source/Developer/IOS/IOSTargetPlatformSettings/IOSTargetPlatformSettings.Build.cs b/Engine/Source/Developer/IOS/IOSTargetPlatformSettings/IOSTargetPlatformSettings.Build.cs
new file mode 100644
--- /dev/null
+++ b/Engine/Source/Developer/IOS/IOSTargetPlatformSettings/IOSTargetPlatformSettings.Build.cs
@@ -0,0 +1,28 @@
+// Copyright Epic Games, Inc. All Rights Reserved.
+
+using UnrealBuildTool;
+
+public class IOSTargetPlatformSettings : ModuleRules
+{
+	public IOSTargetPlatformSettings(ReadOnlyTargetRules Target) : base(Target)
+	{
+		BinariesSubFolder = "IOS";
+		// We need a short name here since this can run afoul very easily of the `MAX_PATH` limit when
+		// combined with building other targets that make use of this.
+		ShortName = "IOSTPSet";
+
+		PrivateDependencyModuleNames.AddRange(
+			new string[]
+			{
+				"Core",
+				"CoreUObject",
+				"TargetPlatform",
+				"DesktopPlatform",
+			}
+		);
+		if (Target.bCompileAgainstEngine)
+		{
+			PrivateDependencyModuleNames.Add("Engine");
+		}
+	}
+}
diff --git a/Engine/Source/Developer/IOS/IOSTargetPlatformSettings/Private/IOSTargetPlatformSettings.cpp b/Engine/Source/Developer/IOS/IOSTargetPlatformSettings/Private/IOSTargetPlatformSettings.cpp
new file mode 100644
--- /dev/null
+++ b/Engine/Source/Developer/IOS/IOSTargetPlatformSettings/Private/IOSTargetPlatformSettings.cpp
@@ -0,0 +1,203 @@
+// Copyright Epic Games, Inc. All Rights Reserved.
+
+/*=============================================================================
+	IOSTargetPlatformSettings.cpp: Implements the FIOSTargetPlatformSettings class.
+=============================================================================*/
+
+#include "IOSTargetPlatformSettings.h"
+#include "HAL/FileManager.h"
+#include "HAL/IConsoleManager.h"
+#include "Misc/Paths.h"
+#include "Misc/App.h"
+
+#if WITH_ENGINE
+#include "Engine/Texture.h"
+#include "TextureResource.h"
+#endif
+
+/* FIOSTargetPlatformSettings structors
+ *****************************************************************************/
+
+FIOSTargetPlatformSettings::FIOSTargetPlatformSettings(bool bInIsTVOS, bool bInIsVisionOS)
+	// override the ini name up in the base classes, which will go into the FTargetPlatformInfo
+	: TTargetPlatformSettingsBase()
+	, bIsTVOS(bInIsTVOS)
+	, bIsVisionOS(bInIsVisionOS)
+	, MobileShadingPath(0)
+	, bDistanceField(false)
+	, bMobileForwardEnableClusteredReflections(false)
+	, bMobileVirtualTextures(false)
+{
+#if WITH_ENGINE
+	TextureLODSettings = nullptr; // TextureLODSettings are registered by the device profile.
+	StaticMeshLODSettings.Initialize(this);
+	GetConfigSystem()->GetBool(TEXT("/Script/Engine.RendererSettings"), TEXT("r.DistanceFields"), bDistanceField, GEngineIni);
+	GetConfigSystem()->GetInt(TEXT("/Script/Engine.RendererSettings"), TEXT("r.Mobile.ShadingPath"), MobileShadingPath, GEngineIni);
+	GetConfigSystem()->GetBool(TEXT("/Script/Engine.RendererSettings"), TEXT("r.Mobile.Forward.EnableClusteredReflections"), bMobileForwardEnableClusteredReflections, GEngineIni);
+	GetConfigSystem()->GetBool(TEXT("/Script/Engine.RendererSettings"), TEXT("r.Mobile.VirtualTextures"), bMobileVirtualTextures, GEngineIni);
+#endif // #if WITH_ENGINE
+}
+
+
+FIOSTargetPlatformSettings::~FIOSTargetPlatformSettings()
+{
+}
+
+/* ITargetPlatform interface
+ *****************************************************************************/
+static bool SupportsMetal()
+{
+	// default to NOT supporting metal
+	bool bSupportsMetal = false;
+	GConfig->GetBool(TEXT("/Script/IOSRuntimeSettings.IOSRuntimeSettings"), TEXT("bSupportsMetal"), bSupportsMetal, GEngineIni);
+	return bSupportsMetal;
+}
+
+static bool SupportsMetalMRT()
+{
+	// default to NOT supporting metal MRT
+	bool bSupportsMetalMRT = false;
+	GConfig->GetBool(TEXT("/Script/IOSRuntimeSettings.IOSRuntimeSettings"), TEXT("bSupportsMetalMRT"), bSupportsMetalMRT, GEngineIni);
+	return bSupportsMetalMRT;
+}
+
+static bool SupportsA8Devices()
+{
+    // default to NOT supporting A8 devices
+    bool bSupportAppleA8 = false;
+    GConfig->GetBool(TEXT("/Script/IOSRuntimeSettings.IOSRuntimeSettings"), TEXT("bSupportAppleA8"), bSupportAppleA8, GEngineIni);
+    return bSupportAppleA8;
+}
+
+bool FIOSTargetPlatformSettings::SupportsFeature( ETargetPlatformFeatures Feature ) const
+{
+	switch (Feature)
+	{
+		case ETargetPlatformFeatures::Packaging:
+		case ETargetPlatformFeatures::DeviceOutputLog:
+			return true;
+
+		case ETargetPlatformFeatures::MobileRendering:
+		case ETargetPlatformFeatures::LowQualityLightmaps:
+			return SupportsMetal();
+			
+		case ETargetPlatformFeatures::DeferredRendering:
+		case ETargetPlatformFeatures::HighQualityLightmaps:
+			return SupportsMetalMRT();
+
+		case ETargetPlatformFeatures::VirtualTextureStreaming:
+			// TODO: should it check r.VirtualTextures for SM5 renderer?
+			return bMobileVirtualTextures;
+
+		case ETargetPlatformFeatures::DistanceFieldAO:
+			return UsesDistanceFields();
+
+		case ETargetPlatformFeatures::NormalmapLAEncodingMode:
+		{
+			static IConsoleVariable* CompressorCVar = IConsoleManager::Get().FindConsoleVariable(TEXT("cook.ASTCTextureCompressor"));
+			const bool bUsesARMCompressor = (CompressorCVar ? (CompressorCVar->GetInt() != 0) : false);
+			return bUsesARMCompressor;
+		}
+
+		case ETargetPlatformFeatures::ShowAsPlatformGroup:
+			return false;
+
+		case ETargetPlatformFeatures::SupportsMultipleConnectionTypes:
+			return true;
+
+		default:
+			break;
+	}
+	
+	return TTargetPlatformSettingsBase<FIOSPlatformProperties>::SupportsFeature(Feature);
+}
+
+void FIOSTargetPlatformSettings::GetAllPossibleShaderFormats( TArray<FName>& OutFormats ) const
+{
+	static FName NAME_SF_METAL(TEXT("SF_METAL"));
+	static FName NAME_SF_METAL_SIM(TEXT("SF_METAL_SIM"));
+	static FName NAME_SF_METAL_MRT(TEXT("SF_METAL_MRT"));
+	static FName NAME_SF_METAL_TVOS(TEXT("SF_METAL_TVOS"));
+	static FName NAME_SF_METAL_MRT_TVOS(TEXT("SF_METAL_MRT_TVOS"));
+
+	if (bIsTVOS)
+	{
+		if (SupportsMetalMRT())
+		{
+			OutFormats.AddUnique(NAME_SF_METAL_MRT_TVOS);
+		}
+
+		// because we are currently using IOS settings, we will always use metal, even if Metal isn't listed as being supported
+		// however, if MetalMRT is specific and Metal is set to false, then we will just use MetalMRT
+		if (SupportsMetal() || !SupportsMetalMRT())
+		{
+			OutFormats.AddUnique(NAME_SF_METAL_TVOS);
+		}
+	}
+	else
+	{
+		if (SupportsMetal())
+		{
+			OutFormats.AddUnique(NAME_SF_METAL);
+
+			bool bEnableSimulatorSupport = false;
+			GConfig->GetBool(TEXT("/Script/IOSRuntimeSettings.IOSRuntimeSettings"), TEXT("bEnableSimulatorSupport"), bEnableSimulatorSupport, GEngineIni);
+			if (bEnableSimulatorSupport)
+			{
+				OutFormats.AddUnique(NAME_SF_METAL_SIM);
+			}
+		}
+
+		if (SupportsMetalMRT())
+		{
+			OutFormats.AddUnique(NAME_SF_METAL_MRT);
+		}
+	}
+}
+
+void FIOSTargetPlatformSettings::GetAllTargetedShaderFormats( TArray<FName>& OutFormats ) const
+{
+	GetAllPossibleShaderFormats(OutFormats);
+}
+
+#if WITH_ENGINE
+
+void FIOSTargetPlatformSettings::GetReflectionCaptureFormats( TArray<FName>& OutFormats ) const
+{
+	const bool bMobileDeferredShading = (MobileShadingPath == 1);
+
+	if (SupportsMetalMRT() || bMobileDeferredShading || bMobileForwardEnableClusteredReflections)
+	{
+		OutFormats.Add(FName(TEXT("FullHDR")));
+	}
+
+	OutFormats.Add(FName(TEXT("EncodedHDR")));
+}
+
+static const FName NameASTC_RGB_HDR(TEXT("ASTC_RGB_HDR"));
+static const FName NameBC5(TEXT("BC5"));
+static const FName NameASTC_NormalLA(TEXT("ASTC_NormalLA"));
+
+// we remap some of the defaults
+static const FName FormatRemap[] =
+{
+	// original				ASTC
+	FName(TEXT("AutoDXT")),	FName(TEXT("ASTC_RGBAuto")),
+	FName(TEXT("DXT1")),	FName(TEXT("ASTC_RGB")),
+	FName(TEXT("DXT5")),	FName(TEXT("ASTC_RGBA")),
+	FName(TEXT("DXT5n")),	FName(TEXT("ASTC_NormalAG")),
+	NameBC5,				FName(TEXT("ASTC_NormalRG")),
+	FName(TEXT("BC4")),		FName(TEXT("ETC2_R11")),
+	FName(TEXT("BC6H")),	NameASTC_RGB_HDR,
+	FName(TEXT("BC7")),		FName(TEXT("ASTC_RGBA_HQ"))
+};
+static const FName NameG8(TEXT("G8"));
+static const FName NameRGBA16F(TEXT("RGBA16F"));
+
+const UTextureLODSettings& FIOSTargetPlatformSettings::GetTextureLODSettings() const
+{
+	return *TextureLODSettings;
+}
+
+#endif // WITH_ENGINE
+
diff --git a/Engine/Source/Developer/IOS/IOSTargetPlatformSettings/Private/IOSTargetPlatformSettings.h b/Engine/Source/Developer/IOS/IOSTargetPlatformSettings/Private/IOSTargetPlatformSettings.h
new file mode 100644
--- /dev/null
+++ b/Engine/Source/Developer/IOS/IOSTargetPlatformSettings/Private/IOSTargetPlatformSettings.h
@@ -0,0 +1,92 @@
+// Copyright Epic Games, Inc. All Rights Reserved.
+
+/*=============================================================================
+	IOSTargetPlatformSettings.h: Declares the FIOSTargetPlatformSettings class.
+=============================================================================*/
+
+#pragma once
+
+#include "CoreMinimal.h"
+#include "Common/TargetPlatformSettingsBase.h"
+#include "IOS/IOSPlatformProperties.h"
+#include "Misc/ConfigCacheIni.h"
+
+#if WITH_ENGINE
+#include "StaticMeshResources.h"
+#endif // WITH_ENGINE
+
+/**
+ * FIOSTargetPlatformSettings
+ */
+class FIOSTargetPlatformSettings : public TTargetPlatformSettingsBase<FIOSPlatformProperties>
+{
+public:
+
+	/**
+	 * Default constructor.
+	 */
+	IOSTARGETPLATFORMSETTINGS_API FIOSTargetPlatformSettings(bool bInISTVOS, bool bInIsVisionOS);
+
+	/**
+	 * Destructor.
+	 */
+	~FIOSTargetPlatformSettings();
+
+public:
+		
+	virtual bool SupportsFeature( ETargetPlatformFeatures Feature ) const override;
+
+	virtual void GetAllPossibleShaderFormats( TArray<FName>& OutFormats ) const override;
+
+	virtual void GetAllTargetedShaderFormats( TArray<FName>& OutFormats ) const override;
+
+#if WITH_ENGINE
+	virtual void GetReflectionCaptureFormats( TArray<FName>& OutFormats ) const override;
+
+	virtual const class FStaticMeshLODSettings& GetStaticMeshLODSettings( ) const override
+	{
+		return StaticMeshLODSettings;
+	}
+
+	virtual const UTextureLODSettings& GetTextureLODSettings() const override;
+
+	virtual void RegisterTextureLODSettings(const UTextureLODSettings* InTextureLODSettings) override
+	{
+		TextureLODSettings = InTextureLODSettings;
+	}
+
+#endif // WITH_ENGINE
+
+	//~ Begin ITargetPlatform Interface
+
+	virtual bool UsesDistanceFields() const override
+	{
+		return bDistanceField;
+	}
+
+private:
+	
+	// true if this is targeting TVOS vs IOS
+	bool bIsTVOS;
+	bool bIsVisionOS;
+
+	// r.Mobile.ShadingPath value
+	int32 MobileShadingPath;
+
+	// true if DistanceField is enabled
+	bool bDistanceField;
+
+	// r.Mobile.Forward.EnableClusteredReflections value
+	bool bMobileForwardEnableClusteredReflections;
+
+	// r.Mobile.VirtualTextures value
+	bool bMobileVirtualTextures;
+
+#if WITH_ENGINE
+	// Holds the cache of the target LOD settings.
+	const UTextureLODSettings* TextureLODSettings;
+
+	// Holds the static mesh LOD settings.
+	FStaticMeshLODSettings StaticMeshLODSettings;
+#endif // WITH_ENGINE
+};
diff --git a/Engine/Source/Developer/IOS/IOSTargetPlatformSettings/Private/IOSTargetPlatformSettingsModule.cpp b/Engine/Source/Developer/IOS/IOSTargetPlatformSettings/Private/IOSTargetPlatformSettingsModule.cpp
new file mode 100644
--- /dev/null
+++ b/Engine/Source/Developer/IOS/IOSTargetPlatformSettings/Private/IOSTargetPlatformSettingsModule.cpp
@@ -0,0 +1,20 @@
+// Copyright Epic Games, Inc. All Rights Reserved.
+
+#include "IOSTargetPlatformSettings.h"
+#include "Interfaces/ITargetPlatformSettingsModule.h"
+#include "Modules/ModuleManager.h"
+
+/**
+ * Module for iOS as a target platform settings
+ */
+class FIOSTargetPlatformSettingsModule : public ITargetPlatformSettingsModule
+{
+public:
+	virtual void GetTargetPlatformSettings(TArray<ITargetPlatformSettings*>& TargetPlatforms) override
+	{
+		TargetPlatforms.Add(new FIOSTargetPlatformSettings(false, false));
+	}
+};
+
+
+IMPLEMENT_MODULE(FIOSTargetPlatformSettingsModule, IOSTargetPlatformSettings);
diff --git a/Engine/Source/Developer/IOS/TVOSTargetPlatform/Private/TVOSTargetPlatformModule.cpp b/Engine/Source/Developer/IOS/TVOSTargetPlatform/Private/TVOSTargetPlatformModule.cpp
--- a/Engine/Source/Developer/IOS/TVOSTargetPlatform/Private/TVOSTargetPlatformModule.cpp
+++ b/Engine/Source/Developer/IOS/TVOSTargetPlatform/Private/TVOSTargetPlatformModule.cpp
@@ -1,9 +1,8 @@
 // Copyright Epic Games, Inc. All Rights Reserved.
 
-#include "IOSTargetPlatform.h"
 #include "Interfaces/ITargetPlatformModule.h"
 #include "Modules/ModuleManager.h"
-
+#include "Common/TargetPlatformBase.h"
 
 /**
  * Module for TVOS as a target platform
@@ -14,14 +13,15 @@ public:
 
 	virtual void GetTargetPlatforms(TArray<ITargetPlatform*>& TargetPlatforms) override
 	{
-		if (FIOSTargetPlatform::IsUsable())
+	}
+
+	virtual void GetTargetPlatforms(TArray<ITargetPlatform*>& TargetPlatforms, TArray<ITargetPlatformSettings*> TargetPlatformSettings, TArray<ITargetPlatformControls*> TargetPlatformControls)
+	{
+		for (ITargetPlatformControls* TargetPlatformControlsIt : TargetPlatformControls)
 		{
-			// add Game and Client TPs
-			TargetPlatforms.Add(new FIOSTargetPlatform(true, false, false));
-			TargetPlatforms.Add(new FIOSTargetPlatform(true, false, true));
+			TargetPlatforms.Add(new FTargetPlatformMerged(TargetPlatformControlsIt->GetTargetPlatformSettings(), TargetPlatformControlsIt));
 		}
 	}
-
 };
 
 
diff --git a/Engine/Source/Developer/IOS/TVOSTargetPlatform/TVOSTargetPlatform.Build.cs b/Engine/Source/Developer/IOS/TVOSTargetPlatform/TVOSTargetPlatform.Build.cs
--- a/Engine/Source/Developer/IOS/TVOSTargetPlatform/TVOSTargetPlatform.Build.cs
+++ b/Engine/Source/Developer/IOS/TVOSTargetPlatform/TVOSTargetPlatform.Build.cs
@@ -13,30 +13,8 @@ public class TVOSTargetPlatform : ModuleRules
 			new string[]
 			{
 				"Core",
-				"CoreUObject",
 				"TargetPlatform",
 				"DesktopPlatform",
-				"LaunchDaemonMessages",
-				"IOSTargetPlatform",
-				"Projects"
-			}
-		);
-
-		PrivateIncludePathModuleNames.AddRange(
-			new string[] {
-			"Messaging",
-			"TargetDeviceServices",
-		}
-		);
-
-		if (Target.bCompileAgainstEngine)
-		{
-			PrivateDependencyModuleNames.Add("Engine");
-		}
-
-		PrivateIncludePaths.AddRange(
-			new string[] {
-			"Developer/IOS/IOSTargetPlatform/Private"
 			}
 		);
 	}
diff --git a/Engine/Source/Developer/IOS/TVOSTargetPlatformControls/Private/TVOSTargetPlatformControlsModule.cpp b/Engine/Source/Developer/IOS/TVOSTargetPlatformControls/Private/TVOSTargetPlatformControlsModule.cpp
new file mode 100644
--- /dev/null
+++ b/Engine/Source/Developer/IOS/TVOSTargetPlatformControls/Private/TVOSTargetPlatformControlsModule.cpp
@@ -0,0 +1,32 @@
+// Copyright Epic Games, Inc. All Rights Reserved.
+
+#include "IOSTargetPlatformControls.h"
+#include "Interfaces/ITargetPlatformControlsModule.h"
+#include "Interfaces/ITargetPlatformSettingsModule.h"
+#include "Modules/ModuleManager.h"
+
+/**
+ * Module for TVOS as a target platform controls
+ */
+class FTVOSTargetPlatformControlsModule	: public ITargetPlatformControlsModule
+{
+public:
+	virtual void GetTargetPlatformControls(TArray<ITargetPlatformControls*>& TargetPlatforms, FName& PlatformSettingsModuleName) override
+	{
+		if (FIOSTargetPlatformControls::IsUsable())
+		{
+			ITargetPlatformSettingsModule* ModuleSettings = FModuleManager::GetModulePtr<ITargetPlatformSettingsModule>(PlatformSettingsModuleName);
+			if (ModuleSettings != nullptr)
+			{
+				TArray<ITargetPlatformSettings*> TargetPlatformSettings = ModuleSettings->GetTargetPlatformSettings();
+				check(TargetPlatformSettings.Num() == 1);
+
+				TargetPlatforms.Add(new FIOSTargetPlatformControls(true, false, false, TargetPlatformSettings[0]));
+				TargetPlatforms.Add(new FIOSTargetPlatformControls(true, false, true, TargetPlatformSettings[0]));
+			}
+		}
+	}
+};
+
+
+IMPLEMENT_MODULE(FTVOSTargetPlatformControlsModule, TVOSTargetPlatformControls);
diff --git a/Engine/Source/Developer/IOS/TVOSTargetPlatformControls/TVOSTargetPlatformControls.Build.cs b/Engine/Source/Developer/IOS/TVOSTargetPlatformControls/TVOSTargetPlatformControls.Build.cs
new file mode 100644
--- /dev/null
+++ b/Engine/Source/Developer/IOS/TVOSTargetPlatformControls/TVOSTargetPlatformControls.Build.cs
@@ -0,0 +1,46 @@
+// Copyright Epic Games, Inc. All Rights Reserved.
+
+using UnrealBuildTool;
+
+public class TVOSTargetPlatformControls : ModuleRules
+{
+	public TVOSTargetPlatformControls(ReadOnlyTargetRules Target) : base(Target)
+	{
+		BinariesSubFolder = "IOS";
+		// We need a short name here since this can run afoul very easily of the `MAX_PATH` limit when
+		// combined with building other targets that make use of this.
+		ShortName = "TVOSTPCon";
+		SDKVersionRelevantPlatforms.Add(UnrealTargetPlatform.TVOS);
+
+		PrivateDependencyModuleNames.AddRange(
+			new string[]
+			{
+				"Core",
+				"CoreUObject",
+				"TargetPlatform",
+				"DesktopPlatform",
+				"LaunchDaemonMessages",
+				"IOSTargetPlatformControls",
+				"Projects"
+			}
+		);
+
+		PrivateIncludePathModuleNames.AddRange(
+			new string[] {
+			"Messaging",
+			"TargetDeviceServices",
+		}
+		);
+
+		if (Target.bCompileAgainstEngine)
+		{
+			PrivateDependencyModuleNames.Add("Engine");
+		}
+
+		PrivateIncludePaths.AddRange(
+			new string[] {
+			"Developer/IOS/IOSTargetPlatformControls/Private"
+			}
+		);
+	}
+}
diff --git a/Engine/Source/Developer/IOS/TVOSTargetPlatformSettings/Private/TVOSTargetPlatformSettingsModule.cpp b/Engine/Source/Developer/IOS/TVOSTargetPlatformSettings/Private/TVOSTargetPlatformSettingsModule.cpp
new file mode 100644
--- /dev/null
+++ b/Engine/Source/Developer/IOS/TVOSTargetPlatformSettings/Private/TVOSTargetPlatformSettingsModule.cpp
@@ -0,0 +1,21 @@
+// Copyright Epic Games, Inc. All Rights Reserved.
+
+#include "IOSTargetPlatformSettings.h"
+#include "Interfaces/ITargetPlatformSettingsModule.h"
+#include "Modules/ModuleManager.h"
+
+/**
+ * Module for TVOS as a target platform settings
+ */
+class FTVOSTargetPlatformSettingsModule	: public ITargetPlatformSettingsModule
+{
+public:
+
+	virtual void GetTargetPlatformSettings(TArray<ITargetPlatformSettings*>& TargetPlatforms) override
+	{
+		TargetPlatforms.Add(new FIOSTargetPlatformSettings(true, false));
+	}
+};
+
+
+IMPLEMENT_MODULE(FTVOSTargetPlatformSettingsModule, TVOSTargetPlatformSettings);
diff --git a/Engine/Source/Developer/IOS/TVOSTargetPlatformSettings/TVOSTargetPlatformSettings.Build.cs b/Engine/Source/Developer/IOS/TVOSTargetPlatformSettings/TVOSTargetPlatformSettings.Build.cs
new file mode 100644
--- /dev/null
+++ b/Engine/Source/Developer/IOS/TVOSTargetPlatformSettings/TVOSTargetPlatformSettings.Build.cs
@@ -0,0 +1,36 @@
+// Copyright Epic Games, Inc. All Rights Reserved.
+
+using UnrealBuildTool;
+
+public class TVOSTargetPlatformSettings : ModuleRules
+{
+	public TVOSTargetPlatformSettings(ReadOnlyTargetRules Target) : base(Target)
+	{
+		BinariesSubFolder = "IOS";
+		// We need a short name here since this can run afoul very easily of the `MAX_PATH` limit when
+		// combined with building other targets that make use of this.
+		ShortName = "TVOSTPSet";
+
+		PrivateDependencyModuleNames.AddRange(
+			new string[]
+			{
+				"Core",
+				"CoreUObject",
+				"TargetPlatform",
+				"DesktopPlatform",
+				"IOSTargetPlatformSettings",
+			}
+		);
+
+		if (Target.bCompileAgainstEngine)
+		{
+			PrivateDependencyModuleNames.Add("Engine");
+		}
+
+		PrivateIncludePaths.AddRange(
+			new string[] {
+			"Developer/IOS/IOSTargetPlatformSettings/Private"
+			}
+		);
+	}
+}
diff --git a/Engine/Source/Programs/UnrealBuildTool/Platform/IOS/UEBuildIOS.cs b/Engine/Source/Programs/UnrealBuildTool/Platform/IOS/UEBuildIOS.cs
--- a/Engine/Source/Programs/UnrealBuildTool/Platform/IOS/UEBuildIOS.cs
+++ b/Engine/Source/Programs/UnrealBuildTool/Platform/IOS/UEBuildIOS.cs
@@ -1022,8 +1022,11 @@ namespace UnrealBuildTool
 		/// <param name="Target">The target being build</param>
 		public override void ModifyModuleRulesForOtherPlatform(string ModuleName, ModuleRules Rules, ReadOnlyTargetRules Target)
 		{
+			bool bIsPlatformAvailableForTarget = UEBuildPlatform.IsPlatformAvailableForTarget(Platform, Target, bIgnoreSDKCheck: true);
+			bool bIsPlatformAvailableForTargetWithSDK = UEBuildPlatform.IsPlatformAvailableForTarget(Platform, Target);
+
 			// don't do any target platform stuff if SDK is not available
-			if (!UEBuildPlatform.IsPlatformAvailableForTarget(Platform, Target))
+			if (!bIsPlatformAvailableForTarget)
 			{
 				return;
 			}
@@ -1037,18 +1040,29 @@ namespace UnrealBuildTool
 					{
 						if (Target.bBuildDeveloperTools)
 						{
-							Rules.DynamicallyLoadedModuleNames.Add("IOSTargetPlatform");
-							Rules.DynamicallyLoadedModuleNames.Add("TVOSTargetPlatform");
+							Rules.DynamicallyLoadedModuleNames.Add("IOSTargetPlatformSettings");
+							Rules.DynamicallyLoadedModuleNames.Add("TVOSTargetPlatformSettings");
+
+							if (bIsPlatformAvailableForTargetWithSDK)
+							{
+								Rules.DynamicallyLoadedModuleNames.Add("IOSTargetPlatform");
+								Rules.DynamicallyLoadedModuleNames.Add("IOSTargetPlatformControls");
+								Rules.DynamicallyLoadedModuleNames.Add("TVOSTargetPlatform");
+								Rules.DynamicallyLoadedModuleNames.Add("TVOSTargetPlatformControls");
+							}
 						}
 					}
 					else if (ModuleName == "TargetPlatform")
 					{
-						bBuildShaderFormats = true;
-						Rules.DynamicallyLoadedModuleNames.Add("TextureFormatASTC");
-						Rules.DynamicallyLoadedModuleNames.Add("TextureFormatETC2");
-						if (Target.bBuildDeveloperTools && Target.bCompileAgainstEngine)
+						if (bIsPlatformAvailableForTargetWithSDK)
 						{
-							Rules.DynamicallyLoadedModuleNames.Add("AudioFormatADPCM");
+							bBuildShaderFormats = true;
+							Rules.DynamicallyLoadedModuleNames.Add("TextureFormatASTC");
+							Rules.DynamicallyLoadedModuleNames.Add("TextureFormatETC2");
+							if (Target.bBuildDeveloperTools && Target.bCompileAgainstEngine)
+							{
+								Rules.DynamicallyLoadedModuleNames.Add("AudioFormatADPCM");
+							}
 						}
 					}
 				}
@@ -1058,17 +1072,25 @@ namespace UnrealBuildTool
 				{
 					if (Target.bForceBuildTargetPlatforms)
 					{
-						Rules.DynamicallyLoadedModuleNames.Add("IOSTargetPlatform");
-						Rules.DynamicallyLoadedModuleNames.Add("TVOSTargetPlatform");
+						Rules.DynamicallyLoadedModuleNames.Add("IOSTargetPlatformSettings");
+						Rules.DynamicallyLoadedModuleNames.Add("TVOSTargetPlatformSettings");
+
+						if (bIsPlatformAvailableForTargetWithSDK)
+						{
+							Rules.DynamicallyLoadedModuleNames.Add("IOSTargetPlatform");
+							Rules.DynamicallyLoadedModuleNames.Add("IOSTargetPlatformControls");
+							Rules.DynamicallyLoadedModuleNames.Add("TVOSTargetPlatform");
+							Rules.DynamicallyLoadedModuleNames.Add("TVOSTargetPlatformControls");
+						}
 					}
 
-					if (bBuildShaderFormats)
+					if (bBuildShaderFormats && bIsPlatformAvailableForTargetWithSDK)
 					{
 						Rules.DynamicallyLoadedModuleNames.Add("MetalShaderFormat");
 					}
 				}
 
-				if (ModuleName == "UnrealEd")
+				if (ModuleName == "UnrealEd" && bIsPlatformAvailableForTargetWithSDK)
 				{
 					Rules.DynamicallyLoadedModuleNames.Add("IOSPlatformEditor");
 				}
