Metal does not support ETC2_11 on 3D textures

diff --git a/Engine/Source/Developer/IOS/IOSTargetPlatformControls/Private/IOSTargetPlatformControls.cpp b/Engine/Source/Developer/IOS/IOSTargetPlatformControls/Private/IOSTargetPlatformControls.cpp
--- a/Engine/Source/Developer/IOS/IOSTargetPlatformControls/Private/IOSTargetPlatformControls.cpp
+++ b/Engine/Source/Developer/IOS/IOSTargetPlatformControls/Private/IOSTargetPlatformControls.cpp
@@ -474,6 +474,7 @@ void FIOSTargetPlatformControls::GetPlatformSpecificProjectAnalytics( TArray<FAn
 
 static const FName NameASTC_RGB_HDR(TEXT("ASTC_RGB_HDR"));
 static const FName NameBC5(TEXT("BC5"));
+static const FName NameBC4(TEXT("BC4"));
 static const FName NameASTC_NormalLA(TEXT("ASTC_NormalLA"));
 
 // we remap some of the defaults
@@ -491,6 +492,7 @@ static const FName FormatRemap[] =
 };
 static const FName NameG8(TEXT("G8"));
 static const FName NameRGBA16F(TEXT("RGBA16F"));
+static const FName NameR16F(TEXT("R16F"));
 
 void FIOSTargetPlatformControls::GetTextureFormats( const UTexture* Texture, TArray< TArray<FName> >& OutFormats) const
 {
@@ -535,6 +537,13 @@ void FIOSTargetPlatformControls::GetTextureFormats( const UTexture* Texture, TAr
 			continue;
 		}
 		
+		// Metal does not support ETC2_11 on 3D textures
+		if (TextureFormatName == NameBC4 && Texture->GetTextureClass() == ETextureClass::Volume)
+		{
+			TextureFormatName = NameR16F;
+			continue;
+		}
+
 		for (int32 RemapIndex = 0; RemapIndex < UE_ARRAY_COUNT(FormatRemap); RemapIndex += 2)
 		{
 			if (TextureFormatName == FormatRemap[RemapIndex])
