IOSTargetPlatformSettings : delete unused texture format strings

diff --git a/Engine/Source/Developer/IOS/IOSTargetPlatformSettings/Private/IOSTargetPlatformSettings.cpp b/Engine/Source/Developer/IOS/IOSTargetPlatformSettings/Private/IOSTargetPlatformSettings.cpp
--- a/Engine/Source/Developer/IOS/IOSTargetPlatformSettings/Private/IOSTargetPlatformSettings.cpp
+++ b/Engine/Source/Developer/IOS/IOSTargetPlatformSettings/Private/IOSTargetPlatformSettings.cpp
@@ -174,26 +174,6 @@ void FIOSTargetPlatformSettings::GetReflectionCaptureFormats( TArray<FName>& Out
 	OutFormats.Add(FName(TEXT("EncodedHDR")));
 }
 
-static const FName NameASTC_RGB_HDR(TEXT("ASTC_RGB_HDR"));
-static const FName NameBC5(TEXT("BC5"));
-static const FName NameASTC_NormalLA(TEXT("ASTC_NormalLA"));
-
-// we remap some of the defaults
-static const FName FormatRemap[] =
-{
-	// original				ASTC
-	FName(TEXT("AutoDXT")),	FName(TEXT("ASTC_RGBAuto")),
-	FName(TEXT("DXT1")),	FName(TEXT("ASTC_RGB")),
-	FName(TEXT("DXT5")),	FName(TEXT("ASTC_RGBA")),
-	FName(TEXT("DXT5n")),	FName(TEXT("ASTC_NormalAG")),
-	NameBC5,				FName(TEXT("ASTC_NormalRG")),
-	FName(TEXT("BC4")),		FName(TEXT("ETC2_R11")),
-	FName(TEXT("BC6H")),	NameASTC_RGB_HDR,
-	FName(TEXT("BC7")),		FName(TEXT("ASTC_RGBA_HQ"))
-};
-static const FName NameG8(TEXT("G8"));
-static const FName NameRGBA16F(TEXT("RGBA16F"));
-
 const UTextureLODSettings& FIOSTargetPlatformSettings::GetTextureLODSettings() const
 {
 	return *TextureLODSettings;
