Make ASTC_RGB_HDR on IOS/Android default to ASTC_RGB if  ASTC HDR is not supported instead of RGBA16F

diff --git a/Engine/Source/Developer/Android/AndroidTargetPlatformControls/Private/AndroidTargetPlatformControls.cpp b/Engine/Source/Developer/Android/AndroidTargetPlatformControls/Private/AndroidTargetPlatformControls.cpp
--- a/Engine/Source/Developer/Android/AndroidTargetPlatformControls/Private/AndroidTargetPlatformControls.cpp
+++ b/Engine/Source/Developer/Android/AndroidTargetPlatformControls/Private/AndroidTargetPlatformControls.cpp
@@ -585,7 +585,7 @@ void FAndroid_ASTCTargetPlatformControls::GetTextureFormats(const UTexture* Text
 		{
 			if (TextureFormatName == AndroidTexFormat::NameASTC_RGB_HDR)
 			{
-				TextureFormatName = AndroidTexFormat::NameRGBA16F;
+				TextureFormatName = AndroidTexFormat::NameASTC_RGB_LDR;
 			}
 		}
 	}
diff --git a/Engine/Source/Developer/Android/AndroidTargetPlatformSettings/Public/AndroidTargetPlatformSettings.h b/Engine/Source/Developer/Android/AndroidTargetPlatformSettings/Public/AndroidTargetPlatformSettings.h
--- a/Engine/Source/Developer/Android/AndroidTargetPlatformSettings/Public/AndroidTargetPlatformSettings.h
+++ b/Engine/Source/Developer/Android/AndroidTargetPlatformSettings/Public/AndroidTargetPlatformSettings.h
@@ -64,6 +64,7 @@ namespace AndroidTexFormat
 	};
 
 	static const FName NameASTC_RGB_HDR(TEXT("ASTC_RGB_HDR"));
+	static const FName NameASTC_RGB_LDR(TEXT("ASTC_RGB"));
 
 	const static FName ASTCRemap[][2] =
 	{
diff --git a/Engine/Source/Developer/IOS/IOSTargetPlatformControls/Private/IOSTargetPlatformControls.cpp b/Engine/Source/Developer/IOS/IOSTargetPlatformControls/Private/IOSTargetPlatformControls.cpp
--- a/Engine/Source/Developer/IOS/IOSTargetPlatformControls/Private/IOSTargetPlatformControls.cpp
+++ b/Engine/Source/Developer/IOS/IOSTargetPlatformControls/Private/IOSTargetPlatformControls.cpp
@@ -473,6 +473,7 @@ void FIOSTargetPlatformControls::GetPlatformSpecificProjectAnalytics( TArray<FAn
 #if WITH_ENGINE
 
 static const FName NameASTC_RGB_HDR(TEXT("ASTC_RGB_HDR"));
+static const FName NameASTC_RGB_LDR(TEXT("ASTC_RGB"));
 static const FName NameBC5(TEXT("BC5"));
 static const FName NameBC4(TEXT("BC4"));
 static const FName NameASTC_NormalLA(TEXT("ASTC_NormalLA"));
@@ -563,7 +564,7 @@ void FIOSTargetPlatformControls::GetTextureFormats( const UTexture* Texture, TAr
 		{
 			if ( TextureFormatName == NameASTC_RGB_HDR )
 			{
-				TextureFormatName = NameRGBA16F;
+				TextureFormatName = NameASTC_RGB_LDR;
 			}
 		}
 	}
