Allow Editor iOS device discovery rescan rate to be configurable
 - change default update rate to 10s
 - settable via AppleRescanRate in BaseEditor.ini
 - remove some commented out code that was never enabled

diff --git a/Engine/Config/BaseEditor.ini b/Engine/Config/BaseEditor.ini
--- a/Engine/Config/BaseEditor.ini
+++ b/Engine/Config/BaseEditor.ini
@@ -667,6 +667,8 @@ bCanModifyPluginsFromBrowser=true
 [EditorBulkData]
 LogCookedStatus=false
 
+[DeviceDiscovery]
+AppleRescanRate=10
 
 [/Script/TextureUtilitiesCommon.TextureImportSettings]
 CompressedFormatForFloatTextures=HDRCompressed_BC6
diff --git a/Engine/Source/Developer/IOS/IOSTargetPlatformControls/Private/IOSDeviceHelper.cpp b/Engine/Source/Developer/IOS/IOSTargetPlatformControls/Private/IOSDeviceHelper.cpp
--- a/Engine/Source/Developer/IOS/IOSTargetPlatformControls/Private/IOSDeviceHelper.cpp
+++ b/Engine/Source/Developer/IOS/IOSTargetPlatformControls/Private/IOSDeviceHelper.cpp
@@ -303,9 +303,6 @@ private:
  */
 DECLARE_MULTICAST_DELEGATE_OneParam(FDeviceNotification, void*)
 
-// recheck once per minute
-#define        RECHECK_COUNTER_RESET            12
-
 class FDeviceQueryTask
 : public FRunnable
 {
@@ -313,7 +310,6 @@ public:
 	FDeviceQueryTask()
 	: Stopping(false)
 	, bCheckDevices(true)
-	, NeedSDKCheck(true)
 	, RetryQuery(5)
 	{}
 	
@@ -324,6 +320,8 @@ public:
 	
 	virtual uint32 Run() override
 	{
+		int32 RecheckRateInt = 10;	// 10s default, which should match "AppleRescanRate" in BaseEditor.ini
+		GConfig->GetInt(TEXT("DeviceDiscovery"), TEXT("AppleRescanRate"), RecheckRateInt, GEditorIni);
 		while (!Stopping)
 		{
 			if (IsEngineExitRequested())
@@ -349,7 +347,10 @@ public:
 				FPlatformProcess::Sleep(1.0f);
 			}
 		}
-		int RecheckCounter = RECHECK_COUNTER_RESET;
+
+		bCheckDevices = true;
+		float RecheckRate = static_cast<float>(RecheckRateInt);
+		double StartTime = FPlatformTime::Seconds();
 		while (!Stopping)
 		{
 			if (IsEngineExitRequested())
@@ -359,36 +360,17 @@ public:
 			if (bCheckDevices)
 			{
 #if WITH_EDITOR
+				// BHP - Turning off device check to prevent it from interfering with packaging
 				if (!IsRunningCommandlet())
+#endif
 				{
-					//if (NeedSDKCheck)
-					//{
-					//    NeedSDKCheck = false;
-					//    FProjectStatus ProjectStatus;
-					//    if (!IProjectManager::Get().QueryStatusForCurrentProject(ProjectStatus) || (!ProjectStatus.IsTargetPlatformSupported(TEXT("IOS")) && !ProjectStatus.IsTargetPlatformSupported(TEXT("TVOS"))))
-					//    {
-					//        Enable(false);
-					//    }
-					//}
-					//else
-					{
-						// BHP - Turning off device check to prevent it from interfering with packaging
-						QueryDevices();
-					}
+					QueryDevices();
 				}
-#else
-				QueryDevices();
-#endif
-			}
-			RecheckCounter--;
-			if (RecheckCounter < 0)
-			{
-				RecheckCounter = RECHECK_COUNTER_RESET;
-				bCheckDevices = true;
-				NeedSDKCheck = true;
 			}
 			
-			FPlatformProcess::Sleep(5.0f);
+			float TimeDelta = FPlatformTime::Seconds() - StartTime;
+			FPlatformProcess::Sleep(FMath::Clamp((RecheckRate-TimeDelta), 0, RecheckRate));
+			StartTime = FPlatformTime::Seconds();
 		}
 		
 		return 0;
@@ -458,7 +440,7 @@ private:
 		FString LibimobileDeviceId = GetLibImobileDeviceExe("idevice_id");
 		if (LibimobileDeviceId.Len() == 0)
 		{
-			UE_LOG(LogIOSDeviceHelper, Log, TEXT("idevice_id (iOS device detection) executable missing. Turning off iOS/tvOS device detection."));
+			UE_LOG(LogIOSDeviceHelper, Warning, TEXT("idevice_id (iOS device detection) executable missing. Turning off iOS/tvOS device detection."));
 			HasDevices = false;
 			Enable(false);
 			if (!HasSimDevices)
@@ -542,7 +524,6 @@ private:
 	
 	bool Stopping;
 	bool bCheckDevices;
-	bool NeedSDKCheck;
 	int RetryQuery;
 	TArray<LibIMobileDevice> CachedDevices;
 	FDeviceNotification DeviceNotification;
