Adding Apple Silicon mac as a target device for running iOS projects

[RN] Adding Apple Silicon mac as a target device for running iOS projects

diff --git a/Engine/Build/BatchFiles/Mac/MakeMacNative.sh b/Engine/Build/BatchFiles/Mac/MakeMacNative.sh
--- a/Engine/Build/BatchFiles/Mac/MakeMacNative.sh
+++ b/Engine/Build/BatchFiles/Mac/MakeMacNative.sh
@@ -1,10 +1,12 @@
 #!/bin/sh
 # Copyright Epic Games, Inc. All Rights Reserved.
 
-IOSBinaryPath=$1
-AppBundleName=$2
+SourceDir=$1
+DestinationDir=$2
+AppName=$3			# including the .app extension
 
-rm -rf "$IOSBinaryPath/MacNative"
-mkdir -p "$IOSBinaryPath/MacNative/$AppBundleName/Wrapper"
-ln -s "Wrapper/$AppBundleName" "$IOSBinaryPath/MacNative/$AppBundleName/WrappedBundle"
-rsync -au "$IOSBinaryPath/$AppBundleName" "$IOSBinaryPath/MacNative/$AppBundleName/Wrapper"
+# Delete previous bundle if any
+rm -rf "$DestinationDir/$AppName"
+mkdir -p "$DestinationDir/$AppName/Wrapper"
+ln -s "Wrapper/$AppName" "$DestinationDir/$AppName/WrappedBundle"
+rsync -au "$SourceDir/$AppName" "$DestinationDir/$AppName/Wrapper"
diff --git a/Engine/Source/Developer/IOS/IOSTargetPlatformControls/Private/IOSDeviceHelper.cpp b/Engine/Source/Developer/IOS/IOSTargetPlatformControls/Private/IOSDeviceHelper.cpp
--- a/Engine/Source/Developer/IOS/IOSTargetPlatformControls/Private/IOSDeviceHelper.cpp
+++ b/Engine/Source/Developer/IOS/IOSTargetPlatformControls/Private/IOSDeviceHelper.cpp
@@ -266,6 +266,37 @@ static TArray<LibIMobileDevice> GetSimulatorDevices()
 	return ToReturn;
 }
 
+// Special fake UUID, this needs to be consistent with AppleToolChainSettings.LocalMacUUID
+const FString FIOSDeviceHelper::LocalMacDummyUUID = TEXT("10CA18AC-10CA18AC10CA18AC");
+
+// Run iOS app natively on local Mac
+static TArray<LibIMobileDevice> GetLocalMacDevices()
+{
+	TArray<LibIMobileDevice> ToReturn;
+
+	LibIMobileDevice ToAdd;
+	ToAdd.IsAuthorized = true;
+	ToAdd.DeviceOSVersion = TEXT("99.99");
+	ToAdd.DeviceUDID = FIOSDeviceHelper::LocalMacDummyUUID;
+	ToAdd.DeviceType = TEXT("iOSAppOnMac");
+	bool bSupportsIPad = true;
+	GConfig->GetBool(TEXT("/Script/IOSRuntimeSettings.IOSRuntimeSettings"), TEXT("bSupportsIPad"), bSupportsIPad, GEngineIni);
+	if (bSupportsIPad)	// Same logic as Xcode, iPad takes priority
+	{
+		ToAdd.DeviceName = TEXT("My Mac (Designed for iPad)");
+	}
+	else
+	{
+		ToAdd.DeviceName = TEXT("My Mac (Designed for iPhone)");
+	}
+	ToAdd.DeviceID = FString::Printf(TEXT("IOS@%s"), *ToAdd.DeviceUDID);
+	ToAdd.DeviceInterface = DeviceConnectionInterface::NoValue;
+
+	ToReturn.Add(ToAdd);
+	
+	return ToReturn;
+}
+
 class FIOSDevice
 {
 public:
@@ -414,6 +445,7 @@ private:
 		else
 		{
 			CallbackInfo.DeviceID = Device.DeviceID;
+			CallbackInfo.DeviceName = Device.DeviceName;
 			CallbackInfo.DeviceUDID = Device.DeviceUDID;
 			CallbackInfo.DeviceInterface = Device.DeviceInterface;
 			CallbackInfo.msgType = 2;
@@ -426,15 +458,23 @@ private:
 	void QueryDevices()
 	{
 		TArray<LibIMobileDevice> SimulatorDevices;
-		bool HasSimDevices = false;
+		bool HasOtherDevices = false;
 		bool HasDevices = true;
+		
+#if PLATFORM_MAC_ARM64
+		// We can run iOS app natively on arm64 Mac, add a special dummy device
+		TArray<LibIMobileDevice> ParsedDevices = GetLocalMacDevices();
+		HasOtherDevices = true;
+#else
+		TArray<LibIMobileDevice> ParsedDevices;
+#endif
 
 		bool bEnableSimulatorSupport = false;
 		GConfig->GetBool(TEXT("/Script/IOSRuntimeSettings.IOSRuntimeSettings"), TEXT("bEnableSimulatorSupport"), bEnableSimulatorSupport, GEngineIni);
 		if (bEnableSimulatorSupport)
 		{
 			SimulatorDevices = GetSimulatorDevices();
-			HasSimDevices = SimulatorDevices.Num() > 0;
+			HasOtherDevices |= SimulatorDevices.Num() > 0;
 		}
 
 		FString LibimobileDeviceId = GetLibImobileDeviceExe("idevice_id");
@@ -443,7 +483,7 @@ private:
 			UE_LOG(LogIOSDeviceHelper, Warning, TEXT("idevice_id (iOS device detection) executable missing. Turning off iOS/tvOS device detection."));
 			HasDevices = false;
 			Enable(false);
-			if (!HasSimDevices)
+			if (!HasOtherDevices)
 			{
 				return;
 			}
@@ -470,7 +510,7 @@ private:
 				}
 				CachedDevices.Empty();
 				HasDevices = false;
-				if (!HasSimDevices)
+				if (!HasOtherDevices)
 				{
 					return;
 				}
@@ -478,10 +518,9 @@ private:
 			RetryQuery = 5;
 		}
 
-		TArray<LibIMobileDevice> ParsedDevices;
 		if (HasDevices)
 		{
-			ParsedDevices = GetLibIMobileDevices();
+			ParsedDevices.Append(GetLibIMobileDevices());
 		}
 
 		ParsedDevices.Append(SimulatorDevices);
diff --git a/Engine/Source/Developer/IOS/IOSTargetPlatformControls/Private/IOSDeviceHelper.h b/Engine/Source/Developer/IOS/IOSTargetPlatformControls/Private/IOSDeviceHelper.h
--- a/Engine/Source/Developer/IOS/IOSTargetPlatformControls/Private/IOSDeviceHelper.h
+++ b/Engine/Source/Developer/IOS/IOSTargetPlatformControls/Private/IOSDeviceHelper.h
@@ -60,6 +60,8 @@ public:
 	 * Suspends/Enables the device connect/disconnect thread
 	 */
 	static void EnableDeviceCheck(bool OnOff);
+	
+	static const FString LocalMacDummyUUID;
 
 private:
     static void DeviceCallback(void*);
diff --git a/Engine/Source/Editor/TurnkeySupport/Private/TurnkeySupportModule.cpp b/Engine/Source/Editor/TurnkeySupport/Private/TurnkeySupportModule.cpp
--- a/Engine/Source/Editor/TurnkeySupport/Private/TurnkeySupportModule.cpp
+++ b/Engine/Source/Editor/TurnkeySupport/Private/TurnkeySupportModule.cpp
@@ -1021,6 +1021,8 @@ static FSlateIcon MakePlatformSdkIconAttribute(FName IniPlatformName, TSharedPtr
 			else
 			{
 				UE_LOG(LogTurnkeySupport, Warning, TEXT("Unknown ConnectionType:%s"), *DeviceProxy->GetConnectionType());
+				// Default to platform icon
+				Icon = FSlateIcon(FAppStyle::Get().GetStyleSetName(), FDataDrivenPlatformInfoRegistry::GetPlatformInfo(IniPlatformName).GetIconStyleName(EPlatformIconSize::Normal));
 			}
 
 			// Verify that the icon exists by checking it's not using the "DefaultBrush" icon resource
@@ -1036,8 +1038,20 @@ static FSlateIcon MakePlatformSdkIconAttribute(FName IniPlatformName, TSharedPtr
 	}
 }
 
+// Special fake UUID used to run iOS app natively on local Apple Silicon mac, see also FIOSDeviceHelper::LocalMacDummyUUID
+static const FString LocalMacDummyUUID = TEXT("10CA18AC-10CA18AC10CA18AC");
+
 static void FormatSdkInfo(const FString& PlatformOrDevice, const FTurnkeySdkInfo& SdkInfo, FText& OutInfo, FText& OutToolTip)
 {
+#if PLATFORM_MAC_ARM64
+	if (PlatformOrDevice == TEXT("IOS@") + LocalMacDummyUUID)
+	{
+		OutInfo = LOCTEXT("SdkInfo_LocalMac", "Run native iOS app on local Mac");
+		OutToolTip = LOCTEXT("SdkInfo_LocalMacToolTip", "Use this device to run native iOS app on local Mac. It is still built for iOS platform");
+		return;
+	}
+#endif
+	
 	TArray<FText> Lines;
 
 	for (TPair<FString, FTurnkeySdkInfo::Version> Pair : SdkInfo.SDKVersions)
@@ -2729,6 +2743,11 @@ void FTurnkeySupportModule::UpdateSdkInfoForProxy(const TSharedRef<ITargetDevice
 
 void FTurnkeySupportModule::UpdateSdkInfoForDevices(TArray<FString> PlatformDeviceIds)
 {
+#if PLATFORM_MAC_ARM64
+	// Dummy local mac device UUID, no need to update
+	PlatformDeviceIds.Remove(TEXT("IOS@") + LocalMacDummyUUID);
+#endif
+	
 	if (PlatformDeviceIds.Num() == 0)
 	{
 		return;
@@ -2869,6 +2888,21 @@ FTurnkeySdkInfo FTurnkeySupportModule::GetSdkInfo(FName PlatformName, bool bBloc
 
 FTurnkeySdkInfo FTurnkeySupportModule::GetSdkInfoForDeviceId(const FString& DeviceId) const
 {
+#if PLATFORM_MAC_ARM64
+	// Dummy info for fake local mac device
+	if (DeviceId == TEXT("IOS@") + LocalMacDummyUUID)
+	{
+		FTurnkeySdkInfo DummyInfo;
+		DummyInfo.Status = ETurnkeyPlatformSdkStatus::Valid;
+		DummyInfo.DeviceStatus = ETurnkeyDeviceStatus::SoftwareValid;
+		DummyInfo.DeviceAutoSoftwareUpdates = ETurnkeyDeviceAutoSoftwareUpdateMode::Unknown;
+		DummyInfo.bCanInstallFullSdk = false;
+		DummyInfo.bCanInstallAutoSdk = false;
+		DummyInfo.bHasBestSdk = true;
+		return DummyInfo;
+	}
+#endif
+	
 	FScopeLock Lock(&GTurnkeySection);
 
 	// return the status, or Unknown info if not known
diff --git a/Engine/Source/Programs/AutomationTool/Apple/ApplePlatform.Automation.cs b/Engine/Source/Programs/AutomationTool/Apple/ApplePlatform.Automation.cs
--- a/Engine/Source/Programs/AutomationTool/Apple/ApplePlatform.Automation.cs
+++ b/Engine/Source/Programs/AutomationTool/Apple/ApplePlatform.Automation.cs
@@ -293,18 +293,6 @@ public abstract class ApplePlatform : Platform
 						Logger.LogInformation("Finalized {App} for running fully self-contained", GetFinalAppPath(Target, SC));
 					}
 
-					if (Params.bMakeMacNative)
-					{
-						String MakeMacNativeSHPath = FileReference.Combine(Unreal.EngineDirectory, "Build/BatchFiles/Mac/MakeMacNative.sh").FullName;
-						String Arg1 = FileReference.Combine(Params.RawProjectPath.Directory, "Binaries/IOS").FullName;  // Path to ios binary folder
-						String Arg2 = GetFinalAppPath(Target, SC).GetDirectoryName();                                   // .app name
-						Logger.LogInformation("Generate {1}/MacNative/{2} that can be run directly on an Apple Silicon Mac", Arg1, Arg2);
-						String Result = Utils.RunLocalProcessAndReturnStdOut("/bin/sh", $"\"{MakeMacNativeSHPath}\" \"{Arg1}\" \"{Arg2}\"");
-						if (Result.Length > 0)
-						{
-							Logger.LogInformation(Result);
-						}
-					}
 					Logger.LogInformation("==================================================================================================================================================");
 				}
 
diff --git a/Engine/Source/Programs/AutomationTool/IOS/IOSPlatform.Automation.cs b/Engine/Source/Programs/AutomationTool/IOS/IOSPlatform.Automation.cs
--- a/Engine/Source/Programs/AutomationTool/IOS/IOSPlatform.Automation.cs
+++ b/Engine/Source/Programs/AutomationTool/IOS/IOSPlatform.Automation.cs
@@ -1865,6 +1865,22 @@ Console.WriteLine("target receipt executable path: {0}", InExecutablePath);
 			AppToDeploy = StagedIPA;
 		}
 
+		if (Params.DeviceNames[0].Contains(AppleToolChainSettings.LocalMacUUID) || Params.bMakeMacNative)
+		{
+			// Deploy on local mac, not on iOS device
+			String MakeMacNativeSHPath = FileReference.Combine(Unreal.EngineDirectory, "Build/BatchFiles/Mac/MakeMacNative.sh").FullName;
+			String Arg1 = FileReference.FromString(AppToDeploy).Directory.FullName;														// Path to staged ios folder
+			String Arg2 = FileReference.Combine(FileReference.FromString(AppToDeploy).Directory.ParentDirectory, "iOSonMac").FullName;	// Path to destination folder
+			String Arg3 = FileReference.FromString(AppToDeploy).GetFileName();															// App name
+			Logger.LogInformation("Generating {1}/{2} to be run natively on an Apple Silicon Mac", Arg2, Arg3);
+			String Result = Utils.RunLocalProcessAndReturnStdOut("/bin/sh", $"\"{MakeMacNativeSHPath}\" \"{Arg1}\" \"{Arg2}\" \"{Arg3}\"");
+			if (Result.Length > 0)
+			{
+				Logger.LogInformation(Result);
+			}
+
+			return;
+		}
 
 		// if iterative deploy, determine the file delta
 		string BundleIdentifier = "";
@@ -2047,7 +2063,15 @@ Console.WriteLine("target receipt executable path: {0}", InExecutablePath);
 
 			string Program;
 			string Arguments;
-			if (OperatingSystem.IsMacOS() &&
+			if (Params.DeviceNames[0].Contains(AppleToolChainSettings.LocalMacUUID))
+			{
+				// Run iOS app natively on Mac, the binary is in StagedBuilds/iOSonMac
+				string IOSonMacDir = Path.Combine(Params.BaseStageDirectory, "iOSonMac");
+				string AppName = Path.GetFileNameWithoutExtension(ClientApp);
+				Program = "open";
+				Arguments = IOSonMacDir + "/" + AppName + ".app --args";
+			}
+			else if (OperatingSystem.IsMacOS() &&
 				GetDevices().Where(Device => Device.Id == Params.DeviceNames[0] && int.Parse(Device.SoftwareVersion.Split('.').First()) < 17).Count() == 0)
 			{
 				// As of iOS17, Apple changed the remote debugger protocol and libimobiledevice has stopped working for launching apps.
diff --git a/Engine/Source/Programs/UnrealBuildTool/ToolChain/AppleToolChain.cs b/Engine/Source/Programs/UnrealBuildTool/ToolChain/AppleToolChain.cs
--- a/Engine/Source/Programs/UnrealBuildTool/ToolChain/AppleToolChain.cs
+++ b/Engine/Source/Programs/UnrealBuildTool/ToolChain/AppleToolChain.cs
@@ -78,6 +78,11 @@ namespace UnrealBuildTool
 		DirectoryReference SDKDir;
 		DirectoryReference? SimulatorSDKDir = null;
 
+		/// <summary>
+		/// Dummy UUID for running iOS binary natively on Mac
+		/// </summary>
+		public static readonly string LocalMacUUID = "10CA18AC-10CA18AC10CA18AC";
+
 		/// <summary>
 		/// Constructor, called by platform sybclasses
 		/// </summary>
@@ -932,8 +937,9 @@ namespace UnrealBuildTool
 				// use generic:
 				//   - for Mac, the destination seems to want name="My Mac", not the real name of the Mac
 				//   - when no destinations were supplied
+				//   - when trying to build iOS app to run on local mac 
 				//   - when using manual codesigning because the main use of this is to update provisioning information (reducing potential for issues)
-				if (Platform == UnrealTargetPlatform.Mac || DestinationIds == null || DestinationIds.Count == 0 || DestinationIds[0] == "" || !bUseAutomaticCodeSigning)
+				if (Platform == UnrealTargetPlatform.Mac || DestinationIds == null || DestinationIds.Count == 0 || DestinationIds[0] == "" || DestinationIds[0].Contains(AppleToolChainSettings.LocalMacUUID) || !bUseAutomaticCodeSigning)
 				{
 					DestinationParam = $"-destination generic/platform=\"{AppleExports.GetDestinationPlatform(Platform, Architectures)}\"";
 				}
