[Shaders] Support remote shader compilation on iOS again.

This was temporarily disabled due to issues with XGE (CL 5340798).
The conditions and environment of shader compilation have changed since then and the target platforms should no longer determine whether remote shader compilation is supported.

#rnx

#FailedValidations Preflight

diff --git a/Engine/Source/Developer/IOS/IOSTargetPlatformControls/Private/IOSTargetPlatformControls.cpp b/Engine/Source/Developer/IOS/IOSTargetPlatformControls/Private/IOSTargetPlatformControls.cpp
--- a/Engine/Source/Developer/IOS/IOSTargetPlatformControls/Private/IOSTargetPlatformControls.cpp
+++ b/Engine/Source/Developer/IOS/IOSTargetPlatformControls/Private/IOSTargetPlatformControls.cpp
@@ -453,14 +453,6 @@ static bool SupportsA8Devices()
     return bSupportAppleA8;
 }
 
-bool FIOSTargetPlatformControls::CanSupportRemoteShaderCompile() const
-{
-	// for 4.22 we are disabling support for XGE Shader compile on IOS
-	bool bRemoteCompilingEnabled = false;
-	GConfig->GetBool(TEXT("/Script/IOSRuntimeSettings.IOSRuntimeSettings"), TEXT("EnableRemoteShaderCompile"), bRemoteCompilingEnabled, GEngineIni);
-	return false; // !bRemoteCompilingEnabled;
-}
-
 void FIOSTargetPlatformControls::GetPlatformSpecificProjectAnalytics( TArray<FAnalyticsEventAttribute>& AnalyticsParamArray ) const
 {
 	TNonDesktopTargetPlatformControlsBase<FIOSPlatformProperties>::GetPlatformSpecificProjectAnalytics( AnalyticsParamArray );
diff --git a/Engine/Source/Developer/IOS/IOSTargetPlatformControls/Private/IOSTargetPlatformControls.h b/Engine/Source/Developer/IOS/IOSTargetPlatformControls/Private/IOSTargetPlatformControls.h
--- a/Engine/Source/Developer/IOS/IOSTargetPlatformControls/Private/IOSTargetPlatformControls.h
+++ b/Engine/Source/Developer/IOS/IOSTargetPlatformControls/Private/IOSTargetPlatformControls.h
@@ -53,8 +53,6 @@ public:
 
 	virtual ITargetDevicePtr GetDevice( const FTargetDeviceId& DeviceId ) override;
 
-	virtual bool CanSupportRemoteShaderCompile() const override;
-
 	virtual bool IsSdkInstalled(bool bProjectHasCode, FString& OutTutorialPath) const override;
 	virtual int32 CheckRequirements(bool bProjectHasCode, EBuildConfiguration Configuration, bool bRequiresAssetNativization, FString& OutTutorialPath, FString& OutDocumentationPath, FText& CustomizedLogMessage) const override;
 
