Preview Json Improvements
- Added IosPreviewDeviceProfileSelectorModule in order to enable Generate/Preview Json on IOS.
- The DeviceSelectors will now return the PreviewPlatform from the DeviceProfile selected and the device properties.
- Modified the way the Preview Json preview platform label name is shown. It now displays Preview Platform Name + Device Profile, instead of just Device profile.
- Confirmation Dialog added when Generating Platform Json

#Synced-CL 50480275

#Synced-CL 50480275

diff --git a/Engine/Config/IOS/BaseIOSEngine.ini b/Engine/Config/IOS/BaseIOSEngine.ini
--- a/Engine/Config/IOS/BaseIOSEngine.ini
+++ b/Engine/Config/IOS/BaseIOSEngine.ini
@@ -79,6 +79,7 @@ s.StreamableStripDebugNameInShipping=true
 
 [DeviceProfileManager]
 DeviceProfileSelectionModule="IOSDeviceProfileSelector"
+PreviewDeviceProfileSelectionModule="IOSPreviewDeviceProfileSelector"
 
 [LocalNotification]
 DefaultPlatformService=IOSLocalNotification
diff --git a/Engine/Plugins/Runtime/AndroidDeviceProfileSelector/Source/AndroidDeviceProfileSelector/AndroidDeviceProfileSelector.Build.cs b/Engine/Plugins/Runtime/AndroidDeviceProfileSelector/Source/AndroidDeviceProfileSelector/AndroidDeviceProfileSelector.Build.cs
--- a/Engine/Plugins/Runtime/AndroidDeviceProfileSelector/Source/AndroidDeviceProfileSelector/AndroidDeviceProfileSelector.Build.cs
+++ b/Engine/Plugins/Runtime/AndroidDeviceProfileSelector/Source/AndroidDeviceProfileSelector/AndroidDeviceProfileSelector.Build.cs
@@ -15,11 +15,11 @@ public class AndroidDeviceProfileSelector : ModuleRules
 		string SecretGuid;
 		Ini.GetString("AndroidDPSBuildSettings", "SecretGuid", out SecretGuid);
 		if (!string.IsNullOrEmpty(SecretGuid))
-        {
+		{
 			PrivateDefinitions.Add("HASH_PEPPER_SECRET_GUID="+ SecretGuid);
 		}
 
-        PublicDependencyModuleNames.AddRange(
+		PublicDependencyModuleNames.AddRange(
 			new string[]
 			{
 				"Core",
@@ -45,6 +45,8 @@ public class AndroidDeviceProfileSelector : ModuleRules
 			PrivateDependencyModuleNames.Add("PIEPreviewDeviceSpecification");
 			PublicDependencyModuleNames.Add("Json");
 			PublicDependencyModuleNames.Add("JsonUtilities");
+			PrivateDependencyModuleNames.Add("RHI");
+			PrivateDependencyModuleNames.Add("ApplicationCore");
 		}
 	}
 }
diff --git a/Engine/Plugins/Runtime/AndroidDeviceProfileSelector/Source/AndroidDeviceProfileSelector/Private/AndroidDeviceProfileSelectorModule.cpp b/Engine/Plugins/Runtime/AndroidDeviceProfileSelector/Source/AndroidDeviceProfileSelector/Private/AndroidDeviceProfileSelectorModule.cpp
--- a/Engine/Plugins/Runtime/AndroidDeviceProfileSelector/Source/AndroidDeviceProfileSelector/Private/AndroidDeviceProfileSelectorModule.cpp
+++ b/Engine/Plugins/Runtime/AndroidDeviceProfileSelector/Source/AndroidDeviceProfileSelector/Private/AndroidDeviceProfileSelectorModule.cpp
@@ -12,6 +12,9 @@
 #include "PIEPreviewDeviceSpecification.h"
 #include "JsonObjectConverter.h"
 #include "Misc/FileHelper.h"
+#include "DeviceProfiles/DeviceProfile.h"
+#include "DeviceProfiles/DeviceProfileManager.h"
+#include "DataDrivenShaderPlatformInfo.h"
 
 #endif
 #include "Misc/ConfigCacheIni.h"
@@ -39,7 +42,7 @@ const FString FAndroidDeviceProfileSelectorModule::GetRuntimeDeviceProfileName()
 }
 
 #if WITH_EDITOR
-void FAndroidDeviceProfileSelectorModule::ExportDeviceParametersToJson(FString& FolderLocation)
+void FAndroidDeviceProfileSelectorModule::ExportDeviceParametersToJson(FString& FolderLocation, TArray<FString>& OutExportedFiles)
 {
 #if WITH_ANDROID_DEVICE_DETECTION
 	IAndroidDeviceDetection* DeviceDetection;
@@ -67,12 +70,16 @@ void FAndroidDeviceProfileSelectorModule::ExportDeviceParametersToJson(FString&
 			for (auto DeviceTuple : Devices)
 			{
 				const FAndroidDeviceInfo& DeviceInfo = DeviceTuple.Value;
-				FString DeviceName = FString::Printf(TEXT("%s_%s(OS%s)"), *DeviceInfo.DeviceBrand, *DeviceInfo.Model, *DeviceInfo.HumanAndroidVersion);
-				if (!AlreadyExported.Find(DeviceName))
+				if (!DeviceInfo.DeviceBrand.IsEmpty() && !DeviceInfo.Model.IsEmpty())
 				{
-					FString ExportPath = FolderLocation / (DeviceName + TEXT(".json"));
-					DeviceDetection->ExportDeviceProfile(ExportPath, DeviceTuple.Key);
-					AlreadyExported.Add(DeviceName);
+					FString DeviceName = FString::Printf(TEXT("%s_%s(OS%s)"), *DeviceInfo.DeviceBrand, *DeviceInfo.Model, *DeviceInfo.HumanAndroidVersion);
+					if (!AlreadyExported.Find(DeviceName))
+					{
+						FString ExportPath = FolderLocation / (DeviceName + TEXT(".json"));
+						OutExportedFiles.Add(ExportPath);
+						DeviceDetection->ExportDeviceProfile(ExportPath, DeviceTuple.Key);
+						AlreadyExported.Add(DeviceName);
+					}
 				}
 			}
 		}
@@ -128,6 +135,39 @@ void FAndroidDeviceProfileSelectorModule::GetDeviceParametersFromJson(FString& J
 		OutDeviceParameters.Add(FAndroidProfileSelectorSourceProperties::SRC_InsetsBottom, FString::Printf(TEXT("%f"), DeviceSpecs.InsetsBottom));
 	}
 }
+
+const EShaderPlatform FAndroidDeviceProfileSelectorModule::GetPreviewShaderPlatform()
+{
+	// ensure SelectorProperties does actually contain our parameters
+	check(FAndroidDeviceProfileSelector::GetSelectorProperties().Num() > 0);
+	FString ProfileName = ANDROID_DEFAULT_DEVICE_PROFILE_NAME;
+	ProfileName = FAndroidDeviceProfileSelector::FindMatchingProfile(ProfileName);
+	EShaderPlatform ParentShaderPlatform = SP_OPENGL_ES3_1_ANDROID;
+	TMap<FName, FString> DeviceProfileCVars = UDeviceProfileManager::Get().GatherDeviceProfileCVars(ProfileName, UDeviceProfileManager::EDeviceProfileMode::DPM_CacheValues);
+	if (const FString* DisableVulkanSupport = DeviceProfileCVars.Find(TEXT("r.Android.DisableVulkanSupport")))
+	{
+		bool bIsDisableVulkanSupport = FCString::Atoi(**DisableVulkanSupport) == 1;
+		if (!bIsDisableVulkanSupport)
+		{
+			ParentShaderPlatform = SP_VULKAN_ES3_1_ANDROID;
+		}
+	}
+
+	EShaderPlatform PreviewPlatform = SP_NumPlatforms;
+	for (int32 PlatformIndex = SP_CUSTOM_PLATFORM_FIRST; PlatformIndex < SP_NumPlatforms; PlatformIndex++)
+	{
+		EShaderPlatform PreviewPlatformToCheck = EShaderPlatform(PlatformIndex);
+		if (FDataDrivenShaderPlatformInfo::IsValid(PreviewPlatformToCheck) &&
+			FDataDrivenShaderPlatformInfo::GetIsPreviewPlatform(PreviewPlatformToCheck) &&
+			FDataDrivenShaderPlatformInfo::GetPreviewShaderPlatformParent(PreviewPlatformToCheck) == ParentShaderPlatform)
+		{
+			PreviewPlatform = PreviewPlatformToCheck;
+			break;
+		}
+	}
+	return PreviewPlatform;
+}
+
 #endif
 
 const FString FAndroidDeviceProfileSelectorModule::GetDeviceProfileName()
diff --git a/Engine/Plugins/Runtime/AndroidDeviceProfileSelector/Source/AndroidDeviceProfileSelector/Private/AndroidDeviceProfileSelectorModule.h b/Engine/Plugins/Runtime/AndroidDeviceProfileSelector/Source/AndroidDeviceProfileSelector/Private/AndroidDeviceProfileSelectorModule.h
--- a/Engine/Plugins/Runtime/AndroidDeviceProfileSelector/Source/AndroidDeviceProfileSelector/Private/AndroidDeviceProfileSelectorModule.h
+++ b/Engine/Plugins/Runtime/AndroidDeviceProfileSelector/Source/AndroidDeviceProfileSelector/Private/AndroidDeviceProfileSelectorModule.h
@@ -17,10 +17,11 @@ public:
 	virtual const FString GetDeviceProfileName() override;
 	virtual const FString GetRuntimeDeviceProfileName() override;
 #if WITH_EDITOR
-	virtual void ExportDeviceParametersToJson(FString& FolderLocation) override;
+	virtual void ExportDeviceParametersToJson(FString& FolderLocation, TArray<FString>& OutExportedFiles) override;
 	virtual bool CanExportDeviceParametersToJson() override;
 	virtual void GetDeviceParametersFromJson(FString& JsonLocation, TMap<FName, FString>& OutDeviceParameters) override;
 	virtual bool CanGetDeviceParametersFromJson() override { return true; }
+	virtual const EShaderPlatform GetPreviewShaderPlatform() override;
 #endif
 	// Set the device parameters this selector will use.
 	virtual void SetSelectorProperties(const TMap<FName, FString>& SelectorProperties) override;
diff --git a/Engine/Plugins/Runtime/IOSDeviceProfileSelector/IOSDeviceProfileSelector.uplugin b/Engine/Plugins/Runtime/IOSDeviceProfileSelector/IOSDeviceProfileSelector.uplugin
--- a/Engine/Plugins/Runtime/IOSDeviceProfileSelector/IOSDeviceProfileSelector.uplugin
+++ b/Engine/Plugins/Runtime/IOSDeviceProfileSelector/IOSDeviceProfileSelector.uplugin
@@ -16,6 +16,11 @@
 	"Installed" : false,
 	"Modules" :
 	[
+		{
+			"Name": "IOSPreviewDeviceProfileSelector",
+			"Type": "Editor",
+			"LoadingPhase": "None"
+		},
 		{
 			"Name" : "IOSDeviceProfileSelector",
 			"Type" : "RuntimeNoCommandlet",
diff --git a/Engine/Plugins/Runtime/IOSDeviceProfileSelector/Source/IOSPreviewDeviceProfileSelector/IOSPreviewDeviceProfileSelector.Build.cs b/Engine/Plugins/Runtime/IOSDeviceProfileSelector/Source/IOSPreviewDeviceProfileSelector/IOSPreviewDeviceProfileSelector.Build.cs
new file mode 100644
--- /dev/null
+++ b/Engine/Plugins/Runtime/IOSDeviceProfileSelector/Source/IOSPreviewDeviceProfileSelector/IOSPreviewDeviceProfileSelector.Build.cs
@@ -0,0 +1,43 @@
+// Copyright Epic Games, Inc. All Rights Reserved.
+using System.IO;
+using UnrealBuildBase;
+
+namespace UnrealBuildTool.Rules
+{
+	public class IOSPreviewDeviceProfileSelector : ModuleRules
+	{
+		public IOSPreviewDeviceProfileSelector(ReadOnlyTargetRules Target) : base(Target)
+		{
+			PublicDependencyModuleNames.AddRange(
+				new string[]
+				{
+					"Core",
+					"RHI",
+				}
+				);
+
+			PrivateDependencyModuleNames.AddRange(
+				new string[]
+				{
+					"Core",
+					"CoreUObject",
+					"Engine",
+					"Json",
+					"JsonUtilities",
+					"PIEPreviewDeviceSpecification",
+				}
+				);
+
+			if (IsPlatformAvailable(UnrealTargetPlatform.IOS) && (Target.Type == TargetType.Editor))
+			{
+				PrivateDefinitions.Add("WITH_IOS_DEVICE_DETECTION=1");
+				PublicIncludePathModuleNames.Add("IOSTargetPlatformControls");
+				PrivateDependencyModuleNames.Add("IOSTargetPlatformControls");
+			}
+			else 
+			{
+				PrivateDefinitions.Add("WITH_IOS_DEVICE_DETECTION=0");
+			}
+		}
+	}
+}
diff --git a/Engine/Plugins/Runtime/IOSDeviceProfileSelector/Source/IOSPreviewDeviceProfileSelector/Private/IOSPreviewDeviceProfileSelectorModule.cpp b/Engine/Plugins/Runtime/IOSDeviceProfileSelector/Source/IOSPreviewDeviceProfileSelector/Private/IOSPreviewDeviceProfileSelectorModule.cpp
new file mode 100644
--- /dev/null
+++ b/Engine/Plugins/Runtime/IOSDeviceProfileSelector/Source/IOSPreviewDeviceProfileSelector/Private/IOSPreviewDeviceProfileSelectorModule.cpp
@@ -0,0 +1,197 @@
+// Copyright Epic Games, Inc. All Rights Reserved.
+
+#include "IOSPreviewDeviceProfileSelectorModule.h"
+#include "HAL/PlatformMisc.h"
+#include "Modules/ModuleManager.h"
+#include "PIEPreviewDeviceSpecification.h"
+#include "JsonObjectConverter.h"
+#include "Dom/JsonObject.h"
+#include "Misc/FileHelper.h"
+#include "Internationalization/Regex.h"
+#include "Misc/ConfigCacheIni.h"
+#include "DataDrivenShaderPlatformInfo.h"
+#if WITH_IOS_DEVICE_DETECTION
+#include "IOSDeviceHelper.h"
+#endif
+
+#define IOS_DEFAULT_DEVICE_PROFILE_NAME TEXT("IOS")
+
+IMPLEMENT_MODULE(FIOSPreviewDeviceProfileSelectorModule, IOSPreviewDeviceProfileSelector);
+
+
+void FIOSPreviewDeviceProfileSelectorModule::StartupModule()
+{
+}
+
+
+void FIOSPreviewDeviceProfileSelectorModule::ShutdownModule()
+{
+}
+
+
+FString const FIOSPreviewDeviceProfileSelectorModule::GetRuntimeDeviceProfileName()
+{
+	// We are not expecting this module to have GetRuntimeDeviceProfileName called directly.
+	// IOS ProfileSelectorModule runtime is in FIOSDeviceProfileSelectorModule.
+	// Use GetDeviceProfileName.
+	checkNoEntry();
+	return FString();
+}
+
+const FString GetDeviceProfileNameFromDeviceModel(FString& DeviceModel)
+{
+	TArray<FString> MappingSections;
+	MappingSections.Add(TEXT("IOSDeviceMappings"));
+	MappingSections.Add(TEXT("IOSFallbackDeviceMappings"));
+	for (FString& DeviceMappingSection : MappingSections)
+	{
+		TArray<FString> Mappings;
+		FConfigCacheIni* PlatformEngineIni = FConfigCacheIni::ForPlatform("IOS");
+		if (PlatformEngineIni && PlatformEngineIni->GetSection(*DeviceMappingSection, Mappings, GDeviceProfilesIni))
+		{
+			for (const FString& MappingString : Mappings)
+			{
+				FString MappingRegex, ProfileName;
+				if (MappingString.Split(TEXT("="), &MappingRegex, &ProfileName))
+				{
+					const FRegexPattern RegexPattern(MappingRegex);
+					FRegexMatcher RegexMatcher(RegexPattern, *DeviceModel);
+					if (RegexMatcher.FindNext())
+					{
+						return ProfileName;
+					}
+				}
+				else
+				{
+					FPlatformMisc::LowLevelOutputDebugStringf(TEXT("Invalid %s: %s") LINE_TERMINATOR, *DeviceMappingSection, *MappingString);
+				}
+			}
+		}
+	}
+
+	return IOS_DEFAULT_DEVICE_PROFILE_NAME;
+}
+
+const FString FIOSPreviewDeviceProfileSelectorModule::GetDeviceProfileName()
+{
+	// ensure SelectorProperties does actually contain our parameters
+	FString IOSDeviceProfileName = IOS_DEFAULT_DEVICE_PROFILE_NAME;
+	check(SelectorProperties.Num() > 0);
+	
+	FString DeviceModel;
+	if (GetSelectorPropertyValue(TEXT("SRC_DeviceModel"), DeviceModel))
+	{
+		IOSDeviceProfileName = GetDeviceProfileNameFromDeviceModel(DeviceModel);
+	}
+	return IOSDeviceProfileName;
+}
+
+#if WITH_EDITOR
+
+bool FIOSPreviewDeviceProfileSelectorModule::CanExportDeviceParametersToJson()
+{ 
+#if WITH_IOS_DEVICE_DETECTION
+	return true; 
+#else
+	return false;
+#endif
+};
+
+#if WITH_IOS_DEVICE_DETECTION
+void ExportDeviceProfile(const FString& OutPath, const FString& DeviceName)
+{
+	// instantiate an FPIEPreviewDeviceSpecifications instance and its values
+	FPIEPreviewDeviceSpecifications DeviceSpecs;
+	DeviceSpecs.DevicePlatform = EPIEPreviewDeviceType::IOS;
+	DeviceSpecs.IOSProperties.DeviceModel = DeviceName;
+	// create a JSon object from the above structure
+	TSharedPtr<FJsonObject> JsonObject = FJsonObjectConverter::UStructToJsonObject<FPIEPreviewDeviceSpecifications>(DeviceSpecs);
+
+	// remove Android and switch fields
+	JsonObject->RemoveField(TEXT("AndroidProperties"));
+	JsonObject->RemoveField(TEXT("switchProperties"));
+
+	// serialize the JSon object to string
+	FString OutputString;
+	TSharedRef<TJsonWriter<>> Writer = TJsonWriterFactory<>::Create(&OutputString);
+	FJsonSerializer::Serialize(JsonObject.ToSharedRef(), Writer);
+
+	// export file to disk
+	FFileHelper::SaveStringToFile(OutputString, *OutPath);
+}
+#endif
+
+void FIOSPreviewDeviceProfileSelectorModule::ExportDeviceParametersToJson(FString& FolderLocation, TArray<FString>& OutExportedFiles)
+{
+#if WITH_IOS_DEVICE_DETECTION
+	TSet<FString> AlreadyExported;
+	TArray<FLibIMobileDevice> ParsedDevices = FIOSDeviceHelper::GetLibIMobileDevices();
+	for (FLibIMobileDevice& Device : ParsedDevices)
+	{
+		FString DeviceModel = Device.DeviceType;
+		if (!AlreadyExported.Find(DeviceModel))
+		{
+			FString ExportPath = FolderLocation / (GetDeviceProfileNameFromDeviceModel(DeviceModel) + TEXT(".json"));
+			OutExportedFiles.Add(ExportPath);
+			ExportDeviceProfile(ExportPath, DeviceModel);
+			AlreadyExported.Add(DeviceModel);
+		}
+	}
+#endif
+}
+
+void FIOSPreviewDeviceProfileSelectorModule::GetDeviceParametersFromJson(FString& JsonLocation, TMap<FName, FString>& OutDeviceParameters)
+{
+	TSharedPtr<FJsonObject> JsonRootObject;
+	FString Json;
+
+	if (FFileHelper::LoadFileToString(Json, *JsonLocation))
+	{
+		TSharedRef<TJsonReader<> > JsonReader = TJsonReaderFactory<>::Create(Json);
+		FJsonSerializer::Deserialize(JsonReader, JsonRootObject);
+	}
+
+	FPIEPreviewDeviceSpecifications DeviceSpecs;
+	if (JsonRootObject.IsValid() && FJsonObjectConverter::JsonAttributesToUStruct(JsonRootObject->Values, FPIEPreviewDeviceSpecifications::StaticStruct(), &DeviceSpecs, 0, 0))
+	{
+		FPIEIOSDeviceProperties& IOSProperties = DeviceSpecs.IOSProperties;
+		OutDeviceParameters.Add(FName(TEXT("SRC_DeviceModel")), IOSProperties.DeviceModel);
+	}
+}
+
+const EShaderPlatform FIOSPreviewDeviceProfileSelectorModule::GetPreviewShaderPlatform()
+{
+	EShaderPlatform ParentShaderPlatform = SP_METAL_ES3_1_IOS;
+
+	EShaderPlatform PreviewPlatform = SP_NumPlatforms;
+	for (int32 PlatformIndex = SP_CUSTOM_PLATFORM_FIRST; PlatformIndex < SP_NumPlatforms; PlatformIndex++)
+	{
+		EShaderPlatform PreviewPlatformToCheck = EShaderPlatform(PlatformIndex);
+		if (FDataDrivenShaderPlatformInfo::IsValid(PreviewPlatformToCheck) &&
+			FDataDrivenShaderPlatformInfo::GetIsPreviewPlatform(PreviewPlatformToCheck) &&
+			FDataDrivenShaderPlatformInfo::GetPreviewShaderPlatformParent(PreviewPlatformToCheck) == ParentShaderPlatform)
+		{
+			PreviewPlatform = PreviewPlatformToCheck;
+			break;
+		}
+	}
+	return PreviewPlatform;
+}
+
+void FIOSPreviewDeviceProfileSelectorModule::SetSelectorProperties(const TMap<FName, FString>& InSelectorProperties)
+{
+	SelectorProperties = InSelectorProperties;
+}
+
+bool FIOSPreviewDeviceProfileSelectorModule::GetSelectorPropertyValue(const FName& PropertyType, FString& PropertyValueOUT)
+{
+	if (const FString* Found = SelectorProperties.Find(PropertyType))
+	{
+		PropertyValueOUT = *Found;
+		return true;
+	}
+
+	return false;
+}
+
+#endif
\ No newline at end of file
diff --git a/Engine/Plugins/Runtime/IOSDeviceProfileSelector/Source/IOSPreviewDeviceProfileSelector/Private/IOSPreviewDeviceProfileSelectorModule.h b/Engine/Plugins/Runtime/IOSDeviceProfileSelector/Source/IOSPreviewDeviceProfileSelector/Private/IOSPreviewDeviceProfileSelectorModule.h
new file mode 100644
--- /dev/null
+++ b/Engine/Plugins/Runtime/IOSDeviceProfileSelector/Source/IOSPreviewDeviceProfileSelector/Private/IOSPreviewDeviceProfileSelectorModule.h
@@ -0,0 +1,45 @@
+// Copyright Epic Games, Inc. All Rights Reserved.
+
+#pragma once
+
+#include "CoreMinimal.h"
+#include "IDeviceProfileSelectorModule.h"
+
+/**
+ * Implements the IOS Device Profile Selector module.
+ */
+class FIOSPreviewDeviceProfileSelectorModule
+	: public IDeviceProfileSelectorModule
+{
+public:
+
+	//~ Begin IDeviceProfileSelectorModule Interface
+	virtual const FString GetDeviceProfileName() override;
+	virtual const FString GetRuntimeDeviceProfileName() override;
+#if WITH_EDITOR
+	virtual bool CanExportDeviceParametersToJson() override;
+	virtual void ExportDeviceParametersToJson(FString& FolderLocation, TArray<FString>& OutExportedFiles) override;
+	virtual void GetDeviceParametersFromJson(FString& JsonLocation, TMap<FName, FString>& OutDeviceParameters) override;
+	virtual bool CanGetDeviceParametersFromJson() override { return true; }
+	virtual const EShaderPlatform GetPreviewShaderPlatform() override;
+	virtual void SetSelectorProperties(const TMap<FName, FString>& InSelectorProperties) override;
+	virtual bool GetSelectorPropertyValue(const FName& PropertyType, FString& PropertyValueOUT) override;
+#endif
+	//~ End IDeviceProfileSelectorModule Interface
+
+
+	//~ Begin IModuleInterface Interface
+	virtual void StartupModule() override;
+	virtual void ShutdownModule() override;
+	//~ End IModuleInterface Interface
+
+	
+	/**
+	 * Virtual destructor.
+	 */
+	virtual ~FIOSPreviewDeviceProfileSelectorModule()
+	{
+	}
+private:
+	TMap<FName, FString> SelectorProperties;
+};
diff --git a/Engine/Source/Developer/IOS/IOSTargetPlatformControls/Private/IOSDeviceHelper.cpp b/Engine/Source/Developer/IOS/IOSTargetPlatformControls/Private/IOSDeviceHelper.cpp
--- a/Engine/Source/Developer/IOS/IOSTargetPlatformControls/Private/IOSDeviceHelper.cpp
+++ b/Engine/Source/Developer/IOS/IOSTargetPlatformControls/Private/IOSDeviceHelper.cpp
@@ -17,63 +17,29 @@
 
 DEFINE_LOG_CATEGORY_STATIC(LogIOSDeviceHelper, Log, All);
 
-enum DeviceConnectionInterface
-{
-	NoValue = 0,
-	USB = 1,
-	Network = 2,
-	Simulator = 3,
-	Max = 4
-};
-
-struct FDeviceNotificationCallbackInformation
-{
-	FString DeviceID;
-	FString DeviceName;
-	FString DeviceUDID;
-	FString ProductType;
-	FString DeviceOSVersion;
-	DeviceConnectionInterface DeviceInterface;
-	uint32 msgType;
-	bool IsAuthorized;
-};
-
-
-struct LibIMobileDevice
-{
-	FString DeviceID;
-	FString DeviceName;
-	FString DeviceUDID;
-	FString DeviceType;
-	FString DeviceOSVersion;
-	DeviceConnectionInterface DeviceInterface;
-	bool IsAuthorized;
-	bool IsDealtWith;
-};
-
-const FString StringifyDeviceConnection(DeviceConnectionInterface Interface)
+const FString StringifyDeviceConnection(EDeviceConnectionInterface Interface)
 {
 	switch(Interface)
 	{
-		case DeviceConnectionInterface::NoValue:
+		case EDeviceConnectionInterface::NoValue:
 			return TEXT("NoValue");
 
-		case DeviceConnectionInterface::Network:
+		case EDeviceConnectionInterface::Network:
 			return TEXT("Network");
 
-		case DeviceConnectionInterface::USB:
+		case EDeviceConnectionInterface::USB:
 			return TEXT("USB");
 
-		case DeviceConnectionInterface::Simulator:
+		case EDeviceConnectionInterface::Simulator:
 			return TEXT("Simulator");
 
 		default:
-			UE_LOG(LogIOSDeviceHelper, Warning, TEXT("Unknown DeviceConnectionInterface type:%d"), Interface);
+			UE_LOG(LogIOSDeviceHelper, Warning, TEXT("Unknown EDeviceConnectionInterface type:%d"), Interface);
 			return TEXT("NoValue");
 	}	
 }
 
-static TArray<LibIMobileDevice> GetLibIMobileDevices()
+TArray<FLibIMobileDevice> FIOSDeviceHelper::GetLibIMobileDevices()
 {
 	FString OutStdOut;
 	FString OutStdErr;
@@ -83,7 +49,7 @@ static TArray<LibIMobileDevice> GetLibIMobileDevices()
 	FPlatformProcess::ExecProcess(*LibimobileDeviceId, TEXT(""), &ReturnCode, &OutStdOut, &OutStdErr, NULL, true);
 	
 	
-	TArray<LibIMobileDevice> ToReturn;
+	TArray<FLibIMobileDevice> ToReturn;
 	// separate out each line
 	TArray<FString> OngoingDeviceIds;
 	
@@ -92,7 +58,7 @@ static TArray<LibIMobileDevice> GetLibIMobileDevices()
 	for (int32 StringIndex = 0; StringIndex < OngoingDeviceIds.Num(); ++StringIndex)
 	{
 		const FString& DeviceUDID = OngoingDeviceIds[StringIndex];
-		DeviceConnectionInterface OngoingDeviceInterface = DeviceConnectionInterface::NoValue;
+		EDeviceConnectionInterface OngoingDeviceInterface = EDeviceConnectionInterface::NoValue;
 		
 		FString OutStdOutInfo;
 		FString OutStdErrInfo;
@@ -102,26 +68,26 @@ static TArray<LibIMobileDevice> GetLibIMobileDevices()
 		
 		if (OngoingDeviceIds[StringIndex].Contains("USB"))
 		{
-			OngoingDeviceInterface = DeviceConnectionInterface::USB;
+			OngoingDeviceInterface = EDeviceConnectionInterface::USB;
 		}
 		else if (OngoingDeviceIds[StringIndex].Contains("Network"))
 		{
-			OngoingDeviceInterface = DeviceConnectionInterface::Network;
+			OngoingDeviceInterface = EDeviceConnectionInterface::Network;
 		}
 		OngoingDeviceIds[StringIndex].Split(TEXT(" "), &OngoingDeviceIds[StringIndex], nullptr, ESearchCase::CaseSensitive, ESearchDir::FromStart);
 		
-		if (OngoingDeviceInterface == DeviceConnectionInterface::USB)
+		if (OngoingDeviceInterface == EDeviceConnectionInterface::USB)
 		{
 			Arguments = "-u " + DeviceUDID;
 		}
-		else if (OngoingDeviceInterface == DeviceConnectionInterface::Network)
+		else if (OngoingDeviceInterface == EDeviceConnectionInterface::Network)
 		{
 			Arguments = "-n -u " + DeviceUDID;
 		}
 		
 		FPlatformProcess::ExecProcess(*LibimobileDeviceInfo, *Arguments, &ReturnCodeInfo, &OutStdOutInfo, &OutStdErrInfo, NULL, true);
 		
-		LibIMobileDevice ToAdd;
+		FLibIMobileDevice ToAdd;
 		
 		// ideviceinfo can fail when the connected device is untrusted. It can be "Pairing dialog response pending (-19)", "Invalid HostID (-21)" or "User denied pairing (-18)"
 		// the only thing we can do is to make sure the Trust popup is correctly displayed.
@@ -157,7 +123,7 @@ static TArray<LibIMobileDevice> GetLibIMobileDevices()
 		}
 		else
 		{
-			if (OngoingDeviceInterface == DeviceConnectionInterface::Network)
+			if (OngoingDeviceInterface == EDeviceConnectionInterface::Network)
 			{
 				DeviceName += " [Wifi]";
 			}
@@ -187,11 +153,11 @@ static TArray<LibIMobileDevice> GetLibIMobileDevices()
 	return ToReturn;
 }
 
-static TArray<LibIMobileDevice> GetSimulatorDevices()
+static TArray<FLibIMobileDevice> GetSimulatorDevices()
 {
 	FString OutStdOut;
 	FString OutStdErr;
-	TArray<LibIMobileDevice> ToReturn;
+	TArray<FLibIMobileDevice> ToReturn;
 
 	int ReturnCode;
 	FString simCommand = TEXT("/usr/bin/xcrun");
@@ -227,7 +193,7 @@ static TArray<LibIMobileDevice> GetSimulatorDevices()
 				FString OSVer = OSSimPair.Key.RightChop(IOSIndex+4);
 				OSVer.ReplaceInline(TEXT("-"), TEXT("."));
 
-				LibIMobileDevice ToAdd;
+				FLibIMobileDevice ToAdd;
 				ToAdd.IsAuthorized = true;
 				ToAdd.DeviceOSVersion = OSVer;
 
@@ -250,7 +216,7 @@ static TArray<LibIMobileDevice> GetSimulatorDevices()
 							
 								FString DeviceID = FString::Printf(TEXT("IOS@%s"), *ToAdd.DeviceUDID);
 								ToAdd.DeviceID = DeviceID;
-								ToAdd.DeviceInterface = DeviceConnectionInterface::Simulator;
+								ToAdd.DeviceInterface = EDeviceConnectionInterface::Simulator;
 							
 								ToReturn.Add(ToAdd);
 							}
@@ -270,11 +236,11 @@ static TArray<LibIMobileDevice> GetSimulatorDevices()
 const FString FIOSDeviceHelper::LocalMacDummyUUID = TEXT("10CA18AC-10CA18AC10CA18AC");
 
 // Run iOS app natively on local Mac
-static TArray<LibIMobileDevice> GetLocalMacDevices()
+static TArray<FLibIMobileDevice> GetLocalMacDevices()
 {
-	TArray<LibIMobileDevice> ToReturn;
+	TArray<FLibIMobileDevice> ToReturn;
 
-	LibIMobileDevice ToAdd;
+	FLibIMobileDevice ToAdd;
 	ToAdd.IsAuthorized = true;
 	ToAdd.DeviceOSVersion = TEXT("99.99");
 	ToAdd.DeviceUDID = FIOSDeviceHelper::LocalMacDummyUUID;
@@ -290,7 +256,7 @@ static TArray<LibIMobileDevice> GetLocalMacDevices()
 		ToAdd.DeviceName = TEXT("My Mac (Designed for iPhone)");
 	}
 	ToAdd.DeviceID = FString::Printf(TEXT("IOS@%s"), *ToAdd.DeviceUDID);
-	ToAdd.DeviceInterface = DeviceConnectionInterface::NoValue;
+	ToAdd.DeviceInterface = EDeviceConnectionInterface::NoValue;
 
 	ToReturn.Add(ToAdd);
 	
@@ -300,7 +266,7 @@ static TArray<LibIMobileDevice> GetLocalMacDevices()
 class FIOSDevice
 {
 public:
-	FIOSDevice(FString InID, FString InName, DeviceConnectionInterface InConnectionType)
+	FIOSDevice(FString InID, FString InName, EDeviceConnectionInterface InConnectionType)
 	: UDID(InID)
 	, Name(InName)
 	, ConnectionType(InConnectionType)
@@ -316,7 +282,7 @@ public:
 		return UDID;
 	}
 	
-	DeviceConnectionInterface ConnectionInterface() const
+	EDeviceConnectionInterface ConnectionInterface() const
 	{
 		return ConnectionType;
 	}
@@ -324,7 +290,7 @@ public:
 private:
 	FString UDID;
 	FString Name;
-	DeviceConnectionInterface ConnectionType;
+	EDeviceConnectionInterface ConnectionType;
 };
 
 /**
@@ -427,7 +393,7 @@ public:
 	
 private:
 	
-	void NotifyDeviceChange(LibIMobileDevice& Device, bool bAdd)
+	void NotifyDeviceChange(FLibIMobileDevice& Device, bool bAdd)
 	{
 		FDeviceNotificationCallbackInformation CallbackInfo;
 		
@@ -457,16 +423,16 @@ private:
 	
 	void QueryDevices()
 	{
-		TArray<LibIMobileDevice> SimulatorDevices;
+		TArray<FLibIMobileDevice> SimulatorDevices;
 		bool HasOtherDevices = false;
 		bool HasDevices = true;
 		
 #if PLATFORM_MAC_ARM64
 		// We can run iOS app natively on arm64 Mac, add a special dummy device
-		TArray<LibIMobileDevice> ParsedDevices = GetLocalMacDevices();
+		TArray<FLibIMobileDevice> ParsedDevices = GetLocalMacDevices();
 		HasOtherDevices = true;
 #else
-		TArray<LibIMobileDevice> ParsedDevices;
+		TArray<FLibIMobileDevice> ParsedDevices;
 #endif
 
 		bool bEnableSimulatorSupport = false;
@@ -504,7 +470,7 @@ private:
 					UE_LOG(LogIOSDeviceHelper, Verbose, TEXT("IOS device listing is disabled for 1 minute (too many failed attempts)!"));
 					Enable(false);
 				}
-				for (LibIMobileDevice device : CachedDevices)
+				for (FLibIMobileDevice device : CachedDevices)
 				{
 					NotifyDeviceChange(device, false);
 				}
@@ -520,15 +486,15 @@ private:
 
 		if (HasDevices)
 		{
-			ParsedDevices.Append(GetLibIMobileDevices());
+			ParsedDevices.Append(FIOSDeviceHelper::GetLibIMobileDevices());
 		}
 
 		ParsedDevices.Append(SimulatorDevices);
 		
 		for (int32 Index = 0; Index < ParsedDevices.Num(); ++Index)
 		{
-			LibIMobileDevice *Found = CachedDevices.FindByPredicate(
-				[&](LibIMobileDevice Element)
+			FLibIMobileDevice *Found = CachedDevices.FindByPredicate(
+				[&](FLibIMobileDevice Element)
 				{
 					return (Element.DeviceUDID == ParsedDevices[Index].DeviceUDID &&
 							Element.DeviceInterface == ParsedDevices[Index].DeviceInterface);
@@ -564,7 +530,7 @@ private:
 	bool Stopping;
 	bool bCheckDevices;
 	int RetryQuery;
-	TArray<LibIMobileDevice> CachedDevices;
+	TArray<FLibIMobileDevice> CachedDevices;
 	FDeviceNotification DeviceNotification;
 };
 
diff --git a/Engine/Source/Developer/IOS/IOSTargetPlatformControls/Private/IOSDeviceHelper.h b/Engine/Source/Developer/IOS/IOSTargetPlatformControls/Private/IOSDeviceHelper.h
deleted file mode 100644
diff --git a/Engine/Source/Developer/IOS/IOSTargetPlatformControls/Public/IOSDeviceHelper.h b/Engine/Source/Developer/IOS/IOSTargetPlatformControls/Public/IOSDeviceHelper.h
new file mode 100644
--- /dev/null
+++ b/Engine/Source/Developer/IOS/IOSTargetPlatformControls/Public/IOSDeviceHelper.h
@@ -0,0 +1,104 @@
+// Copyright Epic Games, Inc. All Rights Reserved.
+
+#pragma once
+
+#include "CoreMinimal.h"
+
+enum EDeviceConnectionInterface
+{
+	NoValue = 0,
+	USB = 1,
+	Network = 2,
+	Simulator = 3,
+	Max = 4
+};
+
+struct FDeviceNotificationCallbackInformation
+{
+	FString DeviceID;
+	FString DeviceName;
+	FString DeviceUDID;
+	FString ProductType;
+	FString DeviceOSVersion;
+	EDeviceConnectionInterface DeviceInterface;
+	uint32 msgType;
+	bool IsAuthorized;
+};
+
+struct FLibIMobileDevice
+{
+	FString DeviceID;
+	FString DeviceName;
+	FString DeviceUDID;
+	FString DeviceType;
+	FString DeviceOSVersion;
+	EDeviceConnectionInterface DeviceInterface;
+	bool IsAuthorized;
+	bool IsDealtWith;
+};
+
+/**
+ * Type definition for shared pointers to instances of FIOSDeviceHelper
+ */
+typedef TSharedPtr<class FIOSDeviceHelper, ESPMode::ThreadSafe> FIOSDeviceHelperPtr;
+
+/**
+ * Type definition for shared references to instances of FIOSDeviceHelper
+ */
+typedef TSharedRef<class FIOSDeviceHelper, ESPMode::ThreadSafe> FIOSDeviceHelperRef;
+
+/**
+ * Delegate type for devices being connected or disconnected from the machine
+ *
+ * The first parameter is newly added or removed device
+ */
+DECLARE_MULTICAST_DELEGATE_OneParam(FOnDeviceConnectEvent, const struct FIOSLaunchDaemonPong&);
+
+/**
+ * Implemented the iOS device helper class
+ */
+class FIOSDeviceHelper
+{
+public:
+	static  void Initialize(bool bIsTVOS);
+	static IOSTARGETPLATFORMCONTROLS_API TArray<FLibIMobileDevice> GetLibIMobileDevices();
+	/**
+	 * Returns a delegate that is executed when a device is connected
+	 *
+	 * @return the delegate
+	 */
+	static FOnDeviceConnectEvent& OnDeviceConnected()
+	{
+		static FOnDeviceConnectEvent DeviceConnectedDelegate;
+		return DeviceConnectedDelegate;
+	}
+	
+	/**
+	 * Returns a delegate the is executed when a device is disconnected
+	 *
+	 * @return the delegate
+	 */
+	static FOnDeviceConnectEvent& OnDeviceDisconnected()
+	{
+		static FOnDeviceConnectEvent DeviceDisconnectedDelegate;
+		return DeviceDisconnectedDelegate;
+	}
+	
+	/**
+	 * Installs an ipa on to the device
+	 */
+	static bool InstallIPAOnDevice(const class FTargetDeviceId& DeviceId, const FString& IPAPath);
+	
+	/**
+	 * Suspends/Enables the device connect/disconnect thread
+	 */
+	static void EnableDeviceCheck(bool OnOff);
+	
+	static const FString LocalMacDummyUUID;
+
+private:
+	static void DeviceCallback(void*);
+	static void DoDeviceConnect(void*);
+	static void DoDeviceDisconnect(void*);
+	static bool MessageTickDelegate(float);
+};
diff --git a/Engine/Source/Editor/LevelEditor/Private/LevelEditor.cpp b/Engine/Source/Editor/LevelEditor/Private/LevelEditor.cpp
--- a/Engine/Source/Editor/LevelEditor/Private/LevelEditor.cpp
+++ b/Engine/Source/Editor/LevelEditor/Private/LevelEditor.cpp
@@ -2232,7 +2232,7 @@ void FLevelEditorModule::BindGlobalLevelEditorCommands()
 						else
 						{
 							FUIAction ActionJsonFile;
-							ActionJsonFile.ExecuteAction = FExecuteAction::CreateStatic(&FLevelEditorActionCallbacks::PreviewJson_Clicked, Item.PlatformName, Item.PreviewShaderPlatformName, PreviewJsonPlatform.FilePath);
+							ActionJsonFile.ExecuteAction = FExecuteAction::CreateStatic(&FLevelEditorActionCallbacks::PreviewJson_Clicked, Item.PlatformName, PreviewJsonPlatform.FilePath);
 							ActionJsonFile.IsActionVisibleDelegate = FIsActionButtonVisible::CreateStatic(&FLevelEditorActionCallbacks::IsPreviewJsonVisible, Item.PlatformName);
 							ActionList.MapAction(PreviewJsonPlatform.CommandInfo, ActionJsonFile);
 						}
diff --git a/Engine/Source/Editor/LevelEditor/Private/LevelEditorActions.cpp b/Engine/Source/Editor/LevelEditor/Private/LevelEditorActions.cpp
--- a/Engine/Source/Editor/LevelEditor/Private/LevelEditorActions.cpp
+++ b/Engine/Source/Editor/LevelEditor/Private/LevelEditorActions.cpp
@@ -552,7 +552,7 @@ void FLevelEditorActionCallbacks::ImportScene_Clicked()
 	FEditorFileUtils::Import();
 }
 
-void FLevelEditorActionCallbacks::PreviewJson_Clicked(FName PlatformName, FName PreviewShaderPlatformName, FString JsonFile)
+void FLevelEditorActionCallbacks::PreviewJson_Clicked(FName PlatformName, FString JsonFile)
 {
 	if (JsonFile.IsEmpty())
 	{
@@ -595,33 +595,44 @@ void FLevelEditorActionCallbacks::PreviewJson_Clicked(FName PlatformName, FName
 				float ConstrainedAspectRatio = DPSelectorModule->GetConstrainedAspectRatio();
 
 				DevileProfileName = DPSelectorModule->GetDeviceProfileName();
-				
-				UDeviceProfile* DevileProfile = UDeviceProfileManager::Get().FindProfile(DevileProfileName, false);
-				if (DevileProfile)
+				EShaderPlatform PreviewShaderPlatform = DPSelectorModule->GetPreviewShaderPlatform();
+
+				// Apply the PreviewPlatform returned by the DPSelectorModule.
+				if (FDataDrivenShaderPlatformInfo::IsValid(PreviewShaderPlatform) && FDataDrivenShaderPlatformInfo::GetIsPreviewPlatform(PreviewShaderPlatform))
 				{
-					EShaderPlatform ShaderPlatform = FDataDrivenShaderPlatformInfo::GetShaderPlatformFromName(PreviewShaderPlatformName);
+					FName PreviewShaderPlatformName = FDataDrivenShaderPlatformInfo::GetName(PreviewShaderPlatform);
 
-					auto GetPreviewFeatureLevelInfo = [&]()
-						{
-							const ERHIFeatureLevel::Type FeatureLevel = GetMaxSupportedFeatureLevel(ShaderPlatform);
-							return FPreviewPlatformInfo(
-								FeatureLevel, 
-								ShaderPlatform, PlatformName, 
-								FDataDrivenShaderPlatformInfo::GetShaderFormat(ShaderPlatform), 
-								DevileProfile->GetFName(), 
-								true, 
-								PreviewShaderPlatformName, 
-								FText::FromName(DevileProfile->GetFName()),
-								ConstrainedAspectRatio,
-								DPSelectorModule->GetSafeZones());
-						};
-
-					FPreviewPlatformInfo PreviewFeatureLevelInfo = GetPreviewFeatureLevelInfo();
-					GEditor->SetPreviewPlatform(PreviewFeatureLevelInfo, false);
+					FName FriendlyDeviceProfilePreviewName = FName(FDataDrivenShaderPlatformInfo::GetFriendlyName(PreviewShaderPlatform).ToString() + TEXT(" ") + DevileProfileName);
+
+					UDeviceProfile* DevileProfile = UDeviceProfileManager::Get().FindProfile(DevileProfileName, false);
+					if (DevileProfile)
+					{
+						auto GetPreviewFeatureLevelInfo = [&]()
+							{
+								const ERHIFeatureLevel::Type FeatureLevel = GetMaxSupportedFeatureLevel(PreviewShaderPlatform);
+								return FPreviewPlatformInfo(
+									FeatureLevel,
+									PreviewShaderPlatform, PlatformName,
+									FDataDrivenShaderPlatformInfo::GetShaderFormat(PreviewShaderPlatform),
+									DevileProfile->GetFName(),
+									true,
+									PreviewShaderPlatformName,
+									FText::FromName(FriendlyDeviceProfilePreviewName),
+									ConstrainedAspectRatio,
+									DPSelectorModule->GetSafeZones());
+							};
+
+						FPreviewPlatformInfo PreviewFeatureLevelInfo = GetPreviewFeatureLevelInfo();
+						GEditor->SetPreviewPlatform(PreviewFeatureLevelInfo, false);
+					}
+					else
+					{
+						UE_LOG(LevelEditorActions, Warning, TEXT("JsonFile %s did not match any MatchProfile"), *JsonFile);
+					}
 				}
 				else
 				{
-					UE_LOG(LevelEditorActions, Warning, TEXT("JsonFile %s did not match any MatchProfile"), *JsonFile);
+					UE_LOG(LevelEditorActions, Warning, TEXT("JsonFile %s did not return a valid Preview Platform"), *JsonFile);
 				}
 			}
 		}
@@ -658,9 +669,7 @@ bool FLevelEditorActionCallbacks::IsGeneratePreviewJsonVisible(FName PlatformNam
 
 void FLevelEditorActionCallbacks::GeneratePreviewJson_Clicked(FString PlatformName)
 {
-
 	FString AbsoluteDebugInfoDirectory = IFileManager::Get().ConvertToAbsolutePathForExternalAppForWrite(*(FPaths::ProjectSavedDir() / TEXT("PreviewJsonDevices") / PlatformName));
-
 	{
 		FScopedSlowTask SlowTask(100.f, NSLOCTEXT("Engine", "GeneratePlatformJson", "Generate Platform Json"), true);
 		SlowTask.Visibility = ESlowTaskVisibility::ForceVisible;
@@ -674,7 +683,18 @@ void FLevelEditorActionCallbacks::GeneratePreviewJson_Clicked(FString PlatformNa
 		{
 			if (IDeviceProfileSelectorModule* DPSelectorModule = FModuleManager::LoadModulePtr<IDeviceProfileSelectorModule>(*DeviceProfileSelectionModule))
 			{
-				DPSelectorModule->ExportDeviceParametersToJson(AbsoluteDebugInfoDirectory);
+				TArray<FString> OutExportedFiles;
+				DPSelectorModule->ExportDeviceParametersToJson(AbsoluteDebugInfoDirectory, OutExportedFiles);
+				
+				if (!OutExportedFiles.IsEmpty())
+				{
+					FString FormatedString = FString::Join(OutExportedFiles, TEXT("\n"));
+					FMessageDialog::Open(EAppMsgType::Ok, FText::Format(LOCTEXT("GeneratePreviewJsonDialog", "{0}"), FText::FromString(FormatedString)), LOCTEXT("GeneratePreviewJsonTitle", "Generated Preview Json Files"));
+				}
+				else
+				{
+					FMessageDialog::Open(EAppMsgType::Ok, LOCTEXT("GeneratePreviewJsonDialog_NoFiles", "No Files Generated"), LOCTEXT("GeneratePreviewJsonTitle", "Generated Preview Json Files"));
+				}
 			}
 		}
 	}
@@ -4647,7 +4667,7 @@ void FLevelEditorCommands::RegisterCommands()
 				GenerateJsonPlatform.CommandInfo = FUICommandInfoDecl(
 					this->AsShared(),
 					FName(*FString::Printf(TEXT("Generate Platform Json for %s"), *Item.PlatformName.ToString())),
-					NSLOCTEXT("GeneratePlatformJson", "Generate Platform Json", "Generate Platform Json..."),
+					NSLOCTEXT("GeneratePlatformJson", "Generate Platform Json", "Generate Platform Json"),
 					NSLOCTEXT("GeneratePlatformJsonDesc", "Generate Platform Json From Connected Devices", "Generate Platform Json From Connected Devices"))
 					.UserInterfaceType(EUserInterfaceActionType::Button)
 					.DefaultChord(FInputChord());
@@ -4655,15 +4675,11 @@ void FLevelEditorCommands::RegisterCommands()
 
 				PlatformToPreviewJsonPlatformOverridesValue = &PlatformToPreviewJsonPlatformOverrides.Add(Item.PlatformName);
 				PlatformToPreviewJsonPlatformOverridesValue->Add(GenerateJsonPlatform);
-			}
-
-			check(PlatformToPreviewJsonPlatformOverridesValue != nullptr);
-			if (!PreviewShaderPlatformNames.Find(Item.PreviewShaderPlatformName))
-			{
+			
 				FLevelEditorCommands::PreviewPlatformCommand PreviewJsonPlatform;
 				PreviewJsonPlatform.CommandInfo = FUICommandInfoDecl(
 					this->AsShared(),
-					FName(*FString::Printf(TEXT("Preview via Json with %s"), *Item.PreviewShaderPlatformName.ToString())),
+					FName(*FString::Printf(TEXT("Preview %s via Json"), *Item.PlatformName.ToString())),
 					NSLOCTEXT("PreviewviaJson", "Preview via Json", "Preview via Json..."),
 					NSLOCTEXT("PreviewviaJsonDesc", "Preview via Json", "Preview via Json"))
 					.UserInterfaceType(EUserInterfaceActionType::Button)
@@ -4689,7 +4705,7 @@ void FLevelEditorCommands::RegisterCommands()
 							FLevelEditorCommands::PreviewPlatformCommand PreviewJsonFilePlatform;
 							PreviewJsonFilePlatform.CommandInfo = FUICommandInfoDecl(
 								this->AsShared(),
-								FName(*FString::Printf(TEXT("Preview %s with Json %s"), *JsonFile, *Item.PreviewShaderPlatformName.ToString())),
+								FName(*FString::Printf(TEXT("Preview %s with Json"), *JsonFile)),
 								FText::Format(NSLOCTEXT("PreviewJson", "Preview Json", "Preview {0}"), FText::FromString(FPaths::GetBaseFilename(JsonFile))),
 								FText::Format(NSLOCTEXT("PreviewJsonDesc", "Preview using Platform Json", "Preview {0}"), FText::FromString(FPaths::GetBaseFilename(JsonFile))))
 								.UserInterfaceType(EUserInterfaceActionType::Check)
@@ -4701,8 +4717,7 @@ void FLevelEditorCommands::RegisterCommands()
 							UniqueJsons.Add(JsonFile);
 						}
 					}
-				}	
-				PreviewShaderPlatformNames.Add(Item.PreviewShaderPlatformName);
+				}
 			}
 		}
 	}
diff --git a/Engine/Source/Editor/LevelEditor/Private/ViewportToolbar/LevelEditorSubmenus.cpp b/Engine/Source/Editor/LevelEditor/Private/ViewportToolbar/LevelEditorSubmenus.cpp
--- a/Engine/Source/Editor/LevelEditor/Private/ViewportToolbar/LevelEditorSubmenus.cpp
+++ b/Engine/Source/Editor/LevelEditor/Private/ViewportToolbar/LevelEditorSubmenus.cpp
@@ -2321,23 +2321,24 @@ FToolMenuEntry CreateFeatureLevelPreviewSubmenu()
 						FText::FromString(PlatformName.ToString()),
 						FText(),
 						FNewToolMenuDelegate::CreateLambda(
-							[CommandList, CommandListJson](UToolMenu* InSubMenu)
+							[CommandList, CommandListJson, PlatformName](UToolMenu* InSubMenu)
 							{
 								for (const FLevelEditorCommands::PreviewPlatformCommand& Command : CommandList)
 								{
 									FToolMenuSection& Section = InSubMenu->FindOrAddSection(
 										Command.SectionName,
-										FText::Format(LOCTEXT("PreviewJson", "{0}"), FText::FromName(Command.SectionName))
+										FText::Format(LOCTEXT("PreviewShaderPlatform", "{0}"), FText::FromName(Command.SectionName))
 									);
 									Section.AddMenuEntry(Command.CommandInfo);
 								}
 
-								if (CommandListJson != nullptr)
+								bool bJsonPreviewEnabled = FLevelEditorActionCallbacks::IsPreviewJsonVisible(PlatformName) || FLevelEditorActionCallbacks::IsGeneratePreviewJsonVisible(PlatformName);
+								if (bJsonPreviewEnabled && CommandListJson != nullptr)
 								{
 									FToolMenuSection& SectionJson = InSubMenu->FindOrAddSection(
 										"PreviewWithJson", LOCTEXT("PreviewWithJsonLabel", "Preview With Json")
 									);
-									TMap<FName, TArray<TSharedPtr<FUICommandInfo>>> SectionNameToCommandList;
+									TArray < TSharedPtr<FUICommandInfo>> CommandArray;
 									for (const FLevelEditorCommands::PreviewPlatformCommand& PreviewJsonPlatform :
 										 *CommandListJson)
 									{
@@ -2347,35 +2348,29 @@ FToolMenuEntry CreateFeatureLevelPreviewSubmenu()
 										}
 										else
 										{
-											SectionNameToCommandList.FindOrAdd(PreviewJsonPlatform.SectionName)
-												.Add(PreviewJsonPlatform.CommandInfo);
+											CommandArray.Add(PreviewJsonPlatform.CommandInfo);
 										}
 									}
 
-									for (auto Iter = SectionNameToCommandList.CreateConstIterator(); Iter; ++Iter)
-									{
-										FName SectionName = Iter.Key();
-										const TArray<TSharedPtr<FUICommandInfo>>& CommandListValue = Iter.Value();
-										SectionJson.AddSubMenu(
-											SectionName,
-											FText::Format(
-												LOCTEXT("PreviewJsonLabel", "Preview {0}"), FText::FromName(SectionName)
-											),
-											FText::Format(
-												LOCTEXT("PreviewJsonTooltip", "Preview {0}"), FText::FromName(SectionName)
-											),
-											FNewToolMenuDelegate::CreateLambda(
-												[CommandListValue](UToolMenu* InSubMenu)
+									SectionJson.AddSubMenu(
+										PlatformName,
+										FText::Format(
+											LOCTEXT("PreviewJsonLabel", "Preview {0} with Json"), FText::FromName(PlatformName)
+										),
+										FText::Format(
+											LOCTEXT("PreviewJsonTooltip", "Preview {0} with Json"), FText::FromName(PlatformName)
+										),
+										FNewToolMenuDelegate::CreateLambda(
+											[CommandArray](UToolMenu* InSubMenu)
+											{
+												FToolMenuSection& Section = InSubMenu->AddSection(NAME_None);
+												for (const TSharedPtr<FUICommandInfo>& Command : CommandArray)
 												{
-													FToolMenuSection& Section = InSubMenu->AddSection(NAME_None);
-													for (const TSharedPtr<FUICommandInfo>& Command : CommandListValue)
-													{
-														Section.AddMenuEntry(Command);
-													}
+													Section.AddMenuEntry(Command);
 												}
-											)
-										);
-									}
+											}
+										)
+									);
 								}
 							}
 						)
diff --git a/Engine/Source/Editor/LevelEditor/Public/LevelEditorActions.h b/Engine/Source/Editor/LevelEditor/Public/LevelEditorActions.h
--- a/Engine/Source/Editor/LevelEditor/Public/LevelEditorActions.h
+++ b/Engine/Source/Editor/LevelEditor/Public/LevelEditorActions.h
@@ -647,6 +647,7 @@ public:
 		FName SectionName;
 		TSharedPtr<FUICommandInfo> CommandInfo;
 		FString FilePath;
+		TSet<FName> PreviewShaderPlatformNames;
 	};
 	
 	TSharedPtr< FUICommandInfo > DisablePlatformPreview;
@@ -793,7 +794,7 @@ public:
 	/**
 	* Called When Preview Json is selected in the Platforms Preview Sub Menu
 	*/
-	static UE_API void PreviewJson_Clicked(FName PlatformName, FName PreviewShaderPlatformName, FString JsonFile);
+	static UE_API void PreviewJson_Clicked(FName PlatformName, FString JsonFile);
 
 	/**
 	* Is Preview Json visible in the Platforms Preview Sub Menu
diff --git a/Engine/Source/Editor/UnrealEd/Classes/Editor/EditorPerProjectUserSettings.h b/Engine/Source/Editor/UnrealEd/Classes/Editor/EditorPerProjectUserSettings.h
--- a/Engine/Source/Editor/UnrealEd/Classes/Editor/EditorPerProjectUserSettings.h
+++ b/Engine/Source/Editor/UnrealEd/Classes/Editor/EditorPerProjectUserSettings.h
@@ -204,6 +204,10 @@ public:
 	UPROPERTY(config)
 	FName PreviewShaderPlatformName;
 	
+	/** The shader platforms friendly Name, or Empty Text if there is no preview platform or no Frienly Name */
+	UPROPERTY(config)
+	FText PreviewShaderPlatformFriendlyName;
+
 	/** Constrained Aspect Ratio of the preview platform */
 	UPROPERTY(config)
 	float PreviewConstrainedAspectRatio;
diff --git a/Engine/Source/Editor/UnrealEd/Private/EditorEngine.cpp b/Engine/Source/Editor/UnrealEd/Private/EditorEngine.cpp
--- a/Engine/Source/Editor/UnrealEd/Private/EditorEngine.cpp
+++ b/Engine/Source/Editor/UnrealEd/Private/EditorEngine.cpp
@@ -6481,6 +6481,7 @@ static void SaveShaderPlatformAsDisabled(FPreviewPlatformInfo& PreviewPlatform)
 	Settings->bPreviewFeatureLevelWasDefault = true;
 	Settings->PreviewDeviceProfileName = NAME_None;
 	Settings->PreviewShaderPlatformName = NAME_None;
+	Settings->PreviewShaderPlatformFriendlyName = FText();
 
 	Settings->SaveConfig();
 
@@ -6491,6 +6492,7 @@ static void SaveShaderPlatformAsDisabled(FPreviewPlatformInfo& PreviewPlatform)
 	Settings->bPreviewFeatureLevelWasDefault = (PreviewPlatform.PreviewFeatureLevel == GMaxRHIFeatureLevel);
 	Settings->PreviewDeviceProfileName = PreviewPlatform.DeviceProfileName;
 	Settings->PreviewShaderPlatformName = PreviewPlatform.PreviewShaderPlatformName;
+	Settings->PreviewShaderPlatformFriendlyName = PreviewPlatform.GetFriendlyName();
 }
 
 void UEditorEngine::SetPreviewPlatform(const FPreviewPlatformInfo& NewPreviewPlatform, bool bSaveSettings)
@@ -6763,7 +6765,7 @@ void UEditorEngine::LoadEditorShaderPlatformPreview()
 			Settings->PreviewDeviceProfileName, 
 			Settings->bPreviewFeatureLevelActive, 
 			PreviewShaderPlatformName, 
-			FText(),
+			Settings->PreviewShaderPlatformFriendlyName,
 			Settings->PreviewConstrainedAspectRatio,
 			Settings->PreviewSafeZones),
 			false);
@@ -6780,6 +6782,7 @@ void UEditorEngine::SaveEditorShaderPlatformPreview()
 	Settings->bPreviewFeatureLevelWasDefault = (PreviewPlatform.PreviewFeatureLevel == GMaxRHIFeatureLevel);
 	Settings->PreviewDeviceProfileName = PreviewPlatform.DeviceProfileName;
 	Settings->PreviewShaderPlatformName = PreviewPlatform.PreviewShaderPlatformName;
+	Settings->PreviewShaderPlatformFriendlyName = PreviewPlatform.GetFriendlyName();
 	Settings->PreviewConstrainedAspectRatio = PreviewPlatform.GetConstrainedAspectRatio();
 	Settings->PreviewSafeZones = PreviewPlatform.GetSafeZones();
 	Settings->PostEditChange();
diff --git a/Engine/Source/Runtime/Engine/Public/IDeviceProfileSelectorModule.h b/Engine/Source/Runtime/Engine/Public/IDeviceProfileSelectorModule.h
--- a/Engine/Source/Runtime/Engine/Public/IDeviceProfileSelectorModule.h
+++ b/Engine/Source/Runtime/Engine/Public/IDeviceProfileSelectorModule.h
@@ -9,6 +9,9 @@
 #include "CoreMinimal.h"
 #include "Modules/ModuleInterface.h"
 #include "DeviceProfiles/DeviceProfileMatching.h"
+#if WITH_EDITOR
+#include "RHIShaderPlatform.h"
+#endif
 
 class UDeviceProfile;
 
@@ -36,7 +39,7 @@ public:
 	 * Save Device Jsons to the location provided
 	 * @param FolderLocation Folder where to store the Jsons
 	 */
-	virtual void ExportDeviceParametersToJson(FString& FolderLocation){}
+	virtual void ExportDeviceParametersToJson(FString& FolderLocation, TArray<FString>& OutExportedFiles){}
 	
 	/**
 	 * Can we export device parameters to Jsons
@@ -68,6 +71,11 @@ public:
 	 *  @return Wheter we can get device parameters from Json
 	 */
 	virtual bool CanGetDeviceParametersFromJson() { return false; }
+
+	/**
+	 * @return Return the PreviewShaderPlatform from DeviceProfile
+	 */
+	virtual const EShaderPlatform GetPreviewShaderPlatform() { return EShaderPlatform::SP_NumPlatforms; }
 #endif
 
 	/**
